PerfReportTool crash fix if a summary table stat isn't present in a particular CSV.

diff --git a/Engine/Source/Programs/CSVTools/PerfReportTool/perfsummaries.cs b/Engine/Source/Programs/CSVTools/PerfReportTool/perfsummaries.cs
--- a/Engine/Source/Programs/CSVTools/PerfReportTool/perfsummaries.cs
+++ b/Engine/Source/Programs/CSVTools/PerfReportTool/perfsummaries.cs
@@ -817,15 +817,21 @@ namespace PerfSummaries
 
             foreach (string unitStat in stats)
             {
-                Hitches.Clear();
                 string StatToCheck = unitStat.Split('(')[0];
-                htmlFile.WriteLine("  <tr><td  bgcolor='#ffffff'><b>" + StatToCheck + "</b></td>");
+				StatSamples statSample = csvStats.GetStat(StatToCheck.ToLower());
+				if (statSample == null)
+				{
+					continue;
+				}
+
+				Hitches.Clear();
+				htmlFile.WriteLine("  <tr><td  bgcolor='#ffffff'><b>" + StatToCheck + "</b></td>");
                 Hitches.Add(StatToCheck);
                 int thresholdIndex = 0;
 
                 foreach (float threshold in HitchThresholds)
-                {
-                    float count = (float)csvStats.GetStat(StatToCheck.ToLower()).GetCountOfFramesOverBudget(threshold);
+                {				
+					float count = (float)statSample.GetCountOfFramesOverBudget(threshold);
                     int numSamples = csvStats.GetStat(StatToCheck.ToLower()).GetNumSamples();
                     // if we have 20k frames in a typical flythrough then 20 frames would be red
                     float redThresholdFor50ms = (float)numSamples / 500.0f;
