fix PerfReportTool help list indentation & line breaks.

diff --git a/Engine/Source/Programs/CSVTools/PerfReportTool/PerfReportTool.cs b/Engine/Source/Programs/CSVTools/PerfReportTool/PerfReportTool.cs
--- a/Engine/Source/Programs/CSVTools/PerfReportTool/PerfReportTool.cs
+++ b/Engine/Source/Programs/CSVTools/PerfReportTool/PerfReportTool.cs
@@ -257,9 +257,9 @@ namespace PerfReportTool
 			"       -maxx <frameNumber>\n" +
 			"       -maxy <value>\n" +
 			"       -graphScale <value>\n" +
-            "       -noStripEvents : if specified, don't strip out excluded events from the stats" +
-            "       -perfLog : output performance logging information" +
-			"		-writeSummaryCsv : if specified, a csv file containing summary information will be generated. Not available in bulk mode." +
+            "       -noStripEvents : if specified, don't strip out excluded events from the stats\n" +
+			"       -perfLog : output performance logging information\n" +
+			"       -writeSummaryCsv : if specified, a csv file containing summary information will be generated. Not available in bulk mode.\n" +
             "\n" +
 			"Optional bulk mode args: (use with -csvdir)\n" +
 			"       -recurse \n" +
