Update CSV binaries, bump CsvToSVG version
Fix event annotations getting a broken symbol instead of an arrow when converted to HTML

diff --git a/Engine/Source/Programs/CSVTools/CSVToSVG/CsvToSVG.cs b/Engine/Source/Programs/CSVTools/CSVToSVG/CsvToSVG.cs
--- a/Engine/Source/Programs/CSVTools/CSVToSVG/CsvToSVG.cs
+++ b/Engine/Source/Programs/CSVTools/CSVToSVG/CsvToSVG.cs
@@ -24,7 +24,7 @@ namespace CSVTools
 {
     class Version
     {
-        private static string VersionString = "2.17";
+        private static string VersionString = "2.18";
         
         public static string Get() { return VersionString; }
     };
@@ -1752,8 +1752,8 @@ namespace CSVTools
                     if (eventX - LastEventX <= 8.5f)
                     {
                         // Add an arrow to indicate events were spaced out 
-                        name = "\u21B7 " + name + " (+" + (ev.Frame - lastFrame) + ")";
-                        eventX = LastEventX + 9.0f;
+						name = "&#x21b7; " + name + " (+" + (ev.Frame - lastFrame) + ")";
+						eventX = LastEventX + 9.0f;
                     }
                     float csvTextX = ToCsvX(eventX + 7, rect, range);
 
diff --git a/Engine/Source/Programs/CSVTools/PerfReportTool/PerfReportTool.cs b/Engine/Source/Programs/CSVTools/PerfReportTool/PerfReportTool.cs
--- a/Engine/Source/Programs/CSVTools/PerfReportTool/PerfReportTool.cs
+++ b/Engine/Source/Programs/CSVTools/PerfReportTool/PerfReportTool.cs
@@ -19,7 +19,7 @@ namespace PerfReportTool
 {
     class Version
     {
-        private static string VersionString = "3.77";
+        private static string VersionString = "3.78";
 
         public static string Get() { return VersionString; }
     };
