improvements to PerfReportTool's PeakSummary type:
 - bulk mode summary page can have colour thresholds & additional stats.
 - automation output csv shows all stats, not just those outside of the main summary.

diff --git a/Engine/Source/Programs/CSVTools/PerfReportTool/PerfReportTool.cs b/Engine/Source/Programs/CSVTools/PerfReportTool/PerfReportTool.cs
--- a/Engine/Source/Programs/CSVTools/PerfReportTool/PerfReportTool.cs
+++ b/Engine/Source/Programs/CSVTools/PerfReportTool/PerfReportTool.cs
@@ -19,7 +19,7 @@ namespace PerfReportTool
 {
     class Version
     {
-        private static string VersionString = "3.78";
+        private static string VersionString = "3.79";
 
         public static string Get() { return VersionString; }
     };
diff --git a/Engine/Source/Programs/CSVTools/PerfReportTool/perfsummaries.cs b/Engine/Source/Programs/CSVTools/PerfReportTool/perfsummaries.cs
--- a/Engine/Source/Programs/CSVTools/PerfReportTool/perfsummaries.cs
+++ b/Engine/Source/Programs/CSVTools/PerfReportTool/perfsummaries.cs
@@ -223,7 +223,11 @@ namespace PerfSummaries
 
         public void ReadStatsFromXML(XElement element)
         {
-            stats = element.Element("stats").Value.Split(',').ToList();
+			XElement statsElement = element.Element("stats");
+			if (statsElement != null)
+			{
+				stats = statsElement.Value.Split(',').ToList();
+			}
             foreach (XElement child in element.Elements())
             {
                 if (child.Name == "capture")
@@ -1030,12 +1034,20 @@ namespace PerfSummaries
     {
         public PeakSummary(XElement element)
         {
-            multipliers = new List<double>();
-            budgets = new List<double>();
-            hidePrefixes = new List<string>();
-            sectionPrefixes = new List<string>();
-            shortenedStatNames = new List<string>();
-            isInMainSummary = new List<bool>();
+			multipliers = new List<double>();
+			budgets = new List<double>();
+			hidePrefixes = new List<string>();
+			sectionPrefixes = new List<string>();
+			shortenedStatNames = new List<string>();
+			isInMainSummary = new List<bool>();
+
+			//read the child elements (mostly for colourThresholds)
+			ReadStatsFromXML(element);
+			if (stats.Count > 0)
+			{
+				throw new System.Exception("<stats> is not supported for Peak summary type"); //...yet
+			}
+
             foreach (XElement child in element.Elements())
             {
                 if (child.Name == "hidePrefix")
@@ -1066,7 +1078,7 @@ namespace PerfSummaries
                 }
             }
 
-            if (!bSummaryStatsFound)
+            if (!bSummaryStatsFound && !bIncludeSummaryCsv)
             {
                 return;
             }
@@ -1101,7 +1113,6 @@ namespace PerfSummaries
                 string statName = stat.Split('(')[0];
 
                 if ((csvStats.Stats.ContainsKey(statName.ToLower())) && // If the main stats table contains this stat AND
-                     (isInMainSummary[i] == isMainSummary) && // If we are in the main summary then this stat appears ONLY in the main summary
                      (sectionPrefix == null || statName.StartsWith(sectionPrefix)) // If there is no hide prefix just display the stat, otherwise, make sure they match.
                    )
                 {
@@ -1118,7 +1129,11 @@ namespace PerfSummaries
                     string peakColour = ColourThresholdList.GetThresholdColour(peak, redValue, orangeValue, yellowValue, greenValue);
                     string averageColour = ColourThresholdList.GetThresholdColour(average, redValue, orangeValue, yellowValue, greenValue);
                     string cleanStatName = stats[i].Replace('/', ' ').Replace("$32$", " ");
-                    htmlFile.WriteLine("    <tr><td bgcolor='#ffffff'>" + cleanStatName + "</td><td bgcolor=" + averageColour + ">" + average.ToString("0") + "</td><td bgcolor=" + peakColour + ">" + peak.ToString("0") + "</td><td bgcolor='#ffffff'>" + budget.ToString("0") + "</td></tr>");
+
+					if (isInMainSummary[i] == isMainSummary) // If we are in the main summary then this stat appears ONLY in the main summary
+					{
+						htmlFile.WriteLine("    <tr><td bgcolor='#ffffff'>" + cleanStatName + "</td><td bgcolor=" + averageColour + ">" + average.ToString("0") + "</td><td bgcolor=" + peakColour + ">" + peak.ToString("0") + "</td><td bgcolor='#ffffff'>" + budget.ToString("0") + "</td></tr>");
+					}
 
 					if (LLMCsvData != null)
 					{
@@ -1143,6 +1158,24 @@ namespace PerfSummaries
 			{
 				return;
 			}
+
+			//update metadata
+			if (metadata != null)
+			{
+				foreach (string statName in stats)
+				{
+					if (!csvStats.Stats.ContainsKey(statName.ToLower()))
+					{
+						continue;
+					}
+
+					var statValue = csvStats.Stats[statName.ToLower()];
+					metadata.Add(statName + " Avg", statValue.average.ToString("0.00"), GetStatColourThresholdList(statName));
+					metadata.Add(statName + " Max", statValue.ComputeMaxValue().ToString("0.00"), GetStatColourThresholdList(statName));
+					metadata.Add(statName + " Min", statValue.ComputeMinValue().ToString("0.00"), GetStatColourThresholdList(statName));
+				}
+			}
+
 			// The first thing we always write is the main summary.
 			WriteStatsToHTML(htmlFile, csvStats, true, null, htmlFileName, bIncludeSummaryCsv);
 
@@ -1153,7 +1186,7 @@ namespace PerfSummaries
             for (i = 0; i < sectionPrefixes.Count(); i++)
             {
                 string currentPrefix = sectionPrefixes[i];
-                WriteStatsToHTML(htmlFile, csvStats, false, currentPrefix, htmlFileName, bIncludeSummaryCsv);
+                WriteStatsToHTML(htmlFile, csvStats, false, currentPrefix, htmlFileName, false);
             }
         }
 
@@ -1182,6 +1215,7 @@ namespace PerfSummaries
             multipliers.Add(multiplier);
             isInMainSummary.Add(bIsInMainSummary);
         }
+		
         public override void PostInit(ReportTypeInfo reportTypeInfo)
         {
             // Find the stats by spinning through the graphs in this reporttype
