PerfReportTool 3.80: Fix edge-case exception when multiple events appear on the same frame as a loading screen event that we're stripping out. Looks like this is causing missing perf reports in some cases

diff --git a/Engine/Source/Programs/CSVTools/CsvStats/CsvStats.cs b/Engine/Source/Programs/CSVTools/CsvStats/CsvStats.cs
--- a/Engine/Source/Programs/CSVTools/CsvStats/CsvStats.cs
+++ b/Engine/Source/Programs/CSVTools/CsvStats/CsvStats.cs
@@ -886,8 +886,12 @@ namespace CSVStats
                     {
                         if (csvEvent.Frame == startIndex)
                         {
-                            // Subsequent events will get offset by this amount
-                            FrameOffset -= endIndex - startIndex;
+							// Check if this is the last event this frame
+							if (i == Events.Count - 1 || Events[i + 1].Frame != csvEvent.Frame)
+							{
+								// Subsequent events will get offset by this amount
+								FrameOffset -= endIndex - startIndex;
+							}
                         }
 
                         if (csvEvent.Frame == endIndex+1)
diff --git a/Engine/Source/Programs/CSVTools/PerfReportTool/PerfReportTool.cs b/Engine/Source/Programs/CSVTools/PerfReportTool/PerfReportTool.cs
--- a/Engine/Source/Programs/CSVTools/PerfReportTool/PerfReportTool.cs
+++ b/Engine/Source/Programs/CSVTools/PerfReportTool/PerfReportTool.cs
@@ -19,7 +19,7 @@ namespace PerfReportTool
 {
     class Version
     {
-        private static string VersionString = "3.79";
+        private static string VersionString = "3.80";
 
         public static string Get() { return VersionString; }
     };
