PerfReportTool 3.82
Add -readallstats option : reads all stats so that any stat can be output to the summary table. Useful with -customtable in bulk mode (off by default)

diff --git a/Engine/Source/Programs/CSVTools/PerfReportTool/PerfReportTool.cs b/Engine/Source/Programs/CSVTools/PerfReportTool/PerfReportTool.cs
--- a/Engine/Source/Programs/CSVTools/PerfReportTool/PerfReportTool.cs
+++ b/Engine/Source/Programs/CSVTools/PerfReportTool/PerfReportTool.cs
@@ -19,7 +19,7 @@ namespace PerfReportTool
 {
     class Version
     {
-        private static string VersionString = "3.81";
+        private static string VersionString = "3.82";
 
         public static string Get() { return VersionString; }
     };
@@ -276,6 +276,7 @@ namespace PerfReportTool
             "       -metadataFilter <key=value,key=value...> : filters based on CSV metadata\n" +
             "       -precacheCount <n> : number of CSV files to precache in the lookahead cache (0 for no precache)\n" +
             "       -precacheThreadCount <n> : number of threads to use for the CSV lookahead cache (default 8)\n" +
+			"       -readAllStats : reads all stats so that any stat can be output to the summary table. Useful with -customtable in bulk mode (off by default)\n" +
 			"       -externalGraphs : enables external graphs (off by default)\n" +
 			"";
 			/*
@@ -734,14 +735,17 @@ namespace PerfReportTool
                 statsToSummarise[i] = statsToSummarise[i].ToLower();
             }
 
-            List<string> ListOfKeys = new List<string>(csvStats.Stats.Keys);
-            for (int i = csvStats.Stats.Keys.Count - 1; i >= 0; i--)
-            {
-                if (!statsToSummarise.Contains(ListOfKeys[i]))
-                {
-                    csvStats.Stats.Remove(ListOfKeys[i]);
-                }
-            }
+			if (!GetBoolArg("readAllStats"))
+			{
+				List<string> ListOfKeys = new List<string>(csvStats.Stats.Keys);
+				for (int i = csvStats.Stats.Keys.Count - 1; i >= 0; i--)
+				{
+					if (!statsToSummarise.Contains(ListOfKeys[i]))
+					{
+						csvStats.Stats.Remove(ListOfKeys[i]);
+					}
+				}
+			}
 
             // Generate CSV metadata
 			if (summaryMetadata != null)
@@ -778,9 +782,19 @@ namespace PerfReportTool
 						summaryMetadata.Add(pair.Key.ToLower(), pair.Value);
 					}
 				}
-            }
 
-            if (htmlFilename != null && !string.IsNullOrEmpty(outputDir))
+				// Add every stat avg value to the metadata
+				if (GetBoolArg("readAllStats") )
+				{
+					foreach ( StatSamples stat in csvStats.Stats.Values )
+					{
+						summaryMetadata.Add( stat.Name.ToLower(), stat.average.ToString());
+					}
+				}
+
+			}
+
+			if (htmlFilename != null && !string.IsNullOrEmpty(outputDir))
             {
                 htmlFilename = Path.Combine(outputDir, htmlFilename);
             }
