PerfReportTool 3.81
- Add support for derived metadata. Specifyable via the reportTypes XML, e.g:
	<derivedMetadataMappings>
		<mapping destName="deviceType" destValue="console" sourceName="platform" sourceValue="ps4"/>
		<mapping destName="deviceType" destValue="console" sourceName="platform" sourceValue="xboxone"/>
		<mapping destName="deviceType" destValue="console" sourceName="platform" sourceValue="switch"/>
		<mapping destName="devicetype" destValue="desktop" sourceName="platform" sourceValue="windows"/>
		<mapping destName="devicetype" destValue="desktop" sourceName="platform" sourceValue="mac"/>
		<mapping destName="devicetype" destValue="mobile" sourceName="platform" sourceValue="android"/>
		<mapping destName="devicetype" destValue="mobile" sourceName="platform" sourceValue="ios"/>
	</derivedMetadataMappings>

diff --git a/Engine/Source/Programs/CSVTools/PerfReportTool/PerfReportTool.cs b/Engine/Source/Programs/CSVTools/PerfReportTool/PerfReportTool.cs
--- a/Engine/Source/Programs/CSVTools/PerfReportTool/PerfReportTool.cs
+++ b/Engine/Source/Programs/CSVTools/PerfReportTool/PerfReportTool.cs
@@ -19,7 +19,7 @@ namespace PerfReportTool
 {
     class Version
     {
-        private static string VersionString = "3.80";
+        private static string VersionString = "3.81";
 
         public static string Get() { return VersionString; }
     };
@@ -74,7 +74,7 @@ namespace PerfReportTool
         }
     };
 
-    class OptionalString
+	class OptionalString
     {
         public OptionalString(string valueIn)
         {
@@ -227,7 +227,7 @@ namespace PerfReportTool
             {
                 // Read CSV stats without data
                 dummyCsvStats = CsvStats.ReadCSVFromLines(lines, null, 0, true);
-                metadata = dummyCsvStats.metaData;
+				metadata = dummyCsvStats.metaData;
             }
         }
         public string filename;
@@ -402,6 +402,7 @@ namespace PerfReportTool
                 try
                 {
                     CachedCsvFile cachedCsvFile = csvFileCache.GetNextCachedCsvFile();
+					reportXML.ApplyDerivedMetadata(cachedCsvFile.metadata);
                     Console.WriteLine("-------------------------------------------------");
                     Console.WriteLine("CSV " + (i+1) + "/" + csvFilenames.Length ) ;
                     Console.WriteLine(csvFilenames[i] );
@@ -800,9 +801,10 @@ namespace PerfReportTool
         CsvStats ReadCsvStats(CachedCsvFile csvFile, int minX, int maxX, out int numFramesStripped)
         {
             CsvStats csvStats = CsvStats.ReadCSVFromLines(csvFile.lines, null);
+			reportXML.ApplyDerivedMetadata(csvStats.metaData);
 
-            // Crop the stats to the range
-            csvStats.CropStats(minX, maxX);
+			// Crop the stats to the range
+			csvStats.CropStats(minX, maxX);
 
             numFramesStripped = 0;
             List<CsvEventStripInfo> eventsToStrip = reportXML.GetCsvEventsToStrip();
@@ -1410,6 +1412,47 @@ namespace PerfReportTool
     };
 
 
+	class DerivedMetadataEntry
+	{
+		public DerivedMetadataEntry(string inSourceName, string inSourceValue, string inDestName, string inDestValue )
+		{
+			sourceName = inSourceName;
+			sourceValue = inSourceValue;
+			destName = inDestName;
+			destValue = inDestValue;
+		}
+		public string sourceName;
+		public string sourceValue;
+		public string destName;
+		public string destValue;
+	};
+
+	class DerivedMetadataMappings
+	{
+		public DerivedMetadataMappings()
+		{
+			entries = new List<DerivedMetadataEntry>();
+		}
+		public void ApplyMapping(CsvMetadata csvMetadata)
+		{
+			if (csvMetadata != null)
+			{
+				foreach (DerivedMetadataEntry entry in entries)
+				{
+					if (csvMetadata.Values.ContainsKey(entry.sourceName.ToLowerInvariant()))
+					{
+						if (csvMetadata.Values[entry.sourceName].ToLowerInvariant() == entry.sourceValue.ToLowerInvariant())
+						{
+							csvMetadata.Values.Add(entry.destName.ToLowerInvariant(), entry.destValue);
+						}
+					}
+				}
+			}
+		}
+		public List<DerivedMetadataEntry> entries;
+	}
+
+
 	class ReportXML
 	{
 		bool IsAbsolutePath(string path)
@@ -1551,8 +1594,27 @@ namespace PerfReportTool
 				}
 			}
 
-            // Read events to strip
-            XElement eventsToStripEl = rootElement.Element("csvEventsToStrip");
+			// Read the derived metadata mappings
+			derivedMetadataMappings = new DerivedMetadataMappings();
+			XElement derivedMetadataMappingsElement = rootElement.Element("derivedMetadataMappings");
+			if (derivedMetadataMappingsElement != null)
+			{
+				foreach (XElement mapping in derivedMetadataMappingsElement.Elements("mapping"))
+				{
+					string sourceName = mapping.GetSafeAttibute<string>("sourceName");
+					string sourceValue = mapping.GetSafeAttibute<string>("sourceValue");
+					string destName = mapping.GetSafeAttibute<string>("destName");
+					string destValue = mapping.GetSafeAttibute<string>("destValue");
+					if (sourceName == null || sourceValue == null || destName == null || destValue == null)
+					{
+						throw new Exception("Derivedmetadata mapping is missing a required attribute!\nRequired attributes: sourceName, sourceValue, destName, destValue.\nXML: "+mapping.ToString());
+					}
+					derivedMetadataMappings.entries.Add(new DerivedMetadataEntry(sourceName, sourceValue, destName, destValue));
+				}
+			}
+
+			// Read events to strip
+			XElement eventsToStripEl = rootElement.Element("csvEventsToStrip");
             if (eventsToStripEl != null)
             {
                 csvEventsToStrip = new List<CsvEventStripInfo>();
@@ -1740,6 +1802,11 @@ namespace PerfReportTool
             return csvEventsToStrip;
         }
 
+		public void ApplyDerivedMetadata(CsvMetadata csvMetadata)
+		{
+			derivedMetadataMappings.ApplyMapping(csvMetadata);
+		}
+
 
         Dictionary<string, SummaryTableInfo> summaryTables;
 
@@ -1749,7 +1816,9 @@ namespace PerfReportTool
 		XElement summaryTablesElement;
 		Dictionary<string, GraphSettings> graphs;
 		Dictionary<string, string> statDisplayNameMapping;
-        List<CsvEventStripInfo> csvEventsToStrip;
+		DerivedMetadataMappings derivedMetadataMappings;
+
+		List<CsvEventStripInfo> csvEventsToStrip;
         string reportTypeXmlFilename;
 	}
 
