Make CSV metadata parsing robust to values with commas, in order to prevent exceptions crashes when reading [desktophwinfo].

A nicer solution would be to strip commas at source in the CSV profiler, but that won't work with existing data

diff --git a/Engine/Source/Programs/CSVTools/CsvStats/CsvStats.cs b/Engine/Source/Programs/CSVTools/CsvStats/CsvStats.cs
--- a/Engine/Source/Programs/CSVTools/CsvStats/CsvStats.cs
+++ b/Engine/Source/Programs/CSVTools/CsvStats/CsvStats.cs
@@ -93,8 +93,17 @@ namespace CSVStats
                     {
                         value = entry;
                     }
-                    Values.Add(key, value);
-                    bIsKey = false;
+
+					if (Values.ContainsKey(key))
+					{
+						// If we're still writing to the current key, append the data onto the end
+						Values[key] += "," + value;
+					}
+					else
+					{
+						Values.Add(key, value);
+					}
+					bIsKey = false;
                 }
 
             }
