PerfReportTool 3.84. LLM improvements
- Fix an exception when a CSV has no metadata (e.g LLM)
- Improve peak summary. Add support for the <stats> element (no longer need to rely on stats being referenced in graphs)
- Add support for wildcards, e.g <stats>*</stats>
- Remove duplicate stats automatically
- Refactor peak summary to avoid separate arrays for stat attributes - everything is now wrapped up in a PeakStatInfo struct
- Modify CsvStats csv parsing to skip stats with only spaces in the name (for some reason LLM does this)

diff --git a/Engine/Source/Programs/CSVTools/CsvStats/CsvStats.cs b/Engine/Source/Programs/CSVTools/CsvStats/CsvStats.cs
--- a/Engine/Source/Programs/CSVTools/CsvStats/CsvStats.cs
+++ b/Engine/Source/Programs/CSVTools/CsvStats/CsvStats.cs
@@ -39,9 +39,12 @@ namespace CSVStats
             return rv;
         }
 
-        public CsvMetadata() { }
+        public CsvMetadata()
+		{
+			Values = new Dictionary<string, string>();
+		}
 
-        public CsvMetadata(string[] csvLine)
+		public CsvMetadata(string[] csvLine)
         {
             Values = new Dictionary<string, string>();
             TextInfo ti = CultureInfo.CurrentCulture.TextInfo;
@@ -1030,9 +1033,13 @@ namespace CSVStats
             StatSamples[] stats = new StatSamples[headings.Length];
             for (int i = 0; i < headings.Length; i++)
             {
-                string heading = headings[i];
+                string heading = headings[i].Trim();
+				if ( heading == "")
+				{
+					continue;
+				}
                 // find the events column (if there is one)
-                if (heading.ToLower() == "events")
+                else if (heading.ToLower() == "events")
                 {
                     eventHeadingIndex = i;
                 }
diff --git a/Engine/Source/Programs/CSVTools/PerfReportTool/PerfReportTool.cs b/Engine/Source/Programs/CSVTools/PerfReportTool/PerfReportTool.cs
--- a/Engine/Source/Programs/CSVTools/PerfReportTool/PerfReportTool.cs
+++ b/Engine/Source/Programs/CSVTools/PerfReportTool/PerfReportTool.cs
@@ -19,7 +19,7 @@ namespace PerfReportTool
 {
     class Version
     {
-        private static string VersionString = "3.83";
+        private static string VersionString = "3.84";
 
         public static string Get() { return VersionString; }
     };
@@ -817,6 +817,10 @@ namespace PerfReportTool
             CsvStats csvStats = CsvStats.ReadCSVFromLines(csvFile.lines, null);
 			reportXML.ApplyDerivedMetadata(csvStats.metaData);
 
+			if ( csvStats.metaData == null )
+			{
+				csvStats.metaData = new CsvMetadata();
+			}
 			csvStats.metaData.Values.Add("csvfilename", csvFile.filename);
 
 			// Crop the stats to the range
@@ -1753,7 +1757,7 @@ namespace PerfReportTool
 
 			foreach (Summary summary in reportTypeInfo.summaries)
 			{
-				summary.PostInit(reportTypeInfo);
+				summary.PostInit(reportTypeInfo, csvFile.dummyCsvStats);
 			}
 			return reportTypeInfo;
 		} 
diff --git a/Engine/Source/Programs/CSVTools/PerfReportTool/perfsummaries.cs b/Engine/Source/Programs/CSVTools/PerfReportTool/perfsummaries.cs
--- a/Engine/Source/Programs/CSVTools/PerfReportTool/perfsummaries.cs
+++ b/Engine/Source/Programs/CSVTools/PerfReportTool/perfsummaries.cs
@@ -216,10 +216,25 @@ namespace PerfSummaries
         public virtual void WriteSummaryData(System.IO.StreamWriter htmlFile, CsvStats csvStats, bool bIncludeSummaryCsv, SummaryMetadata metadata, string htmlFileName)
         { }
 
-        public virtual void PostInit(ReportTypeInfo reportTypeInfo)
+        public virtual void PostInit(ReportTypeInfo reportTypeInfo, CsvStats csvStats)
         {
-
-        }
+			// Resolve wildcards and remove duplicates
+			List<string> newStats = new List<string>();
+			Dictionary<string, bool> uniqueStats = new Dictionary<string, bool>();
+			foreach (string statStr in stats)
+			{
+				List<StatSamples> statsMatching = csvStats.GetStatsMatchingString(statStr);
+				foreach (StatSamples stat in statsMatching)
+				{
+					if (!uniqueStats.ContainsKey(stat.Name))
+					{
+						uniqueStats.Add(stat.Name, true);
+						newStats.Add(stat.Name);
+					}
+				}
+			}
+			stats = newStats;
+		}
 
         public void ReadStatsFromXML(XElement element)
         {
@@ -672,7 +687,7 @@ namespace PerfSummaries
                 statsCsvFile.Close();
             }
         }
-        public override void PostInit(ReportTypeInfo reportTypeInfo)
+        public override void PostInit(ReportTypeInfo reportTypeInfo, CsvStats csvStats)
         {
         }
 
@@ -763,7 +778,7 @@ namespace PerfSummaries
                 metadata.Add(metadataKey, eventCount.ToString(), thresholdList);
             }
         }
-        public override void PostInit(ReportTypeInfo reportTypeInfo)
+        public override void PostInit(ReportTypeInfo reportTypeInfo, CsvStats csvStats)
         {
         }
         string[] events;
@@ -1034,19 +1049,11 @@ namespace PerfSummaries
     {
         public PeakSummary(XElement element)
         {
-			multipliers = new List<double>();
-			budgets = new List<double>();
 			hidePrefixes = new List<string>();
 			sectionPrefixes = new List<string>();
-			shortenedStatNames = new List<string>();
-			isInMainSummary = new List<bool>();
 
 			//read the child elements (mostly for colourThresholds)
 			ReadStatsFromXML(element);
-			if (stats.Count > 0)
-			{
-				throw new System.Exception("<stats> is not supported for Peak summary type"); //...yet
-			}
 
             foreach (XElement child in element.Elements())
             {
@@ -1071,7 +1078,9 @@ namespace PerfSummaries
 
             for (int j = 0; j < stats.Count; j++)
             {
-                if ((isInMainSummary[j] == isMainSummary))
+				PeakStatInfo statInfo = getOrAddStatInfo(stats[j]);
+
+				if (statInfo.isInMainSummary == isMainSummary)
                 {
                     bSummaryStatsFound = true;
                     break;
@@ -1112,15 +1121,17 @@ namespace PerfSummaries
             {
                 string statName = stat.Split('(')[0];
 
-                if ((csvStats.Stats.ContainsKey(statName.ToLower())) && // If the main stats table contains this stat AND
+				PeakStatInfo statInfo = getOrAddStatInfo(stat);
+
+				if ((csvStats.Stats.ContainsKey(statName.ToLower())) && // If the main stats table contains this stat AND
                      (sectionPrefix == null || statName.StartsWith(sectionPrefix)) // If there is no hide prefix just display the stat, otherwise, make sure they match.
                    )
                 {
                     // Do the calculations for the averages and peak, and then write it to the table along with the budget.
                     StatSamples csvStat = csvStats.Stats[statName.ToLower()];
-                    double peak = (double)csvStat.ComputeMaxValue() * multipliers[i];
-                    double average = (double)csvStat.average * multipliers[i];
-                    double budget = budgets[i];
+                    double peak = (double)csvStat.ComputeMaxValue() * statInfo.multiplier;
+                    double average = (double)csvStat.average * statInfo.multiplier;
+                    double budget = statInfo.budget;
 
                     float redValue = (float)budget * 1.5f;
                     float orangeValue = (float)budget * 1.25f;
@@ -1130,7 +1141,7 @@ namespace PerfSummaries
                     string averageColour = ColourThresholdList.GetThresholdColour(average, redValue, orangeValue, yellowValue, greenValue);
                     string cleanStatName = stats[i].Replace('/', ' ').Replace("$32$", " ");
 
-					if (isInMainSummary[i] == isMainSummary) // If we are in the main summary then this stat appears ONLY in the main summary
+					if (statInfo.isInMainSummary == isMainSummary) // If we are in the main summary then this stat appears ONLY in the main summary
 					{
 						htmlFile.WriteLine("    <tr><td bgcolor='#ffffff'>" + cleanStatName + "</td><td bgcolor=" + averageColour + ">" + average.ToString("0") + "</td><td bgcolor=" + peakColour + ">" + peak.ToString("0") + "</td><td bgcolor='#ffffff'>" + budget.ToString("0") + "</td></tr>");
 					}
@@ -1190,33 +1201,19 @@ namespace PerfSummaries
             }
         }
 
-        void AddStat(string stat, double budget, double multiplier, bool bIsInMainSummary)
+
+
+        void AddStat(string statName, double budget, double multiplier, bool bIsInMainSummary)
         {
-            // Find the best (longest) prefix which matches this stat, and strip it off
-            int bestPrefixIndex = -1;
-            int bestPrefixLength = 0;
-            string shortStatName = stat;
-            for (int i = 0; i < hidePrefixes.Count; i++)
-            {
-                string prefix = hidePrefixes[i];
-                if (stat.ToLower().StartsWith(prefix) && prefix.Length > bestPrefixLength)
-                {
-                    bestPrefixIndex = i;
-                    bestPrefixLength = prefix.Length;
-                }
-            }
-            if (bestPrefixIndex >= 0)
-            {
-                shortStatName = stat.Substring(bestPrefixLength);
-            }
-            shortenedStatNames.Add(shortStatName);
-            stats.Add(stat);
-            budgets.Add(budget);
-            multipliers.Add(multiplier);
-            isInMainSummary.Add(bIsInMainSummary);
-        }
+            stats.Add(statName);
+
+			PeakStatInfo info = getOrAddStatInfo(statName);
+			info.isInMainSummary = bIsInMainSummary;
+			info.multiplier = multiplier;
+			info.budget = budget;
+		}
 		
-        public override void PostInit(ReportTypeInfo reportTypeInfo)
+        public override void PostInit(ReportTypeInfo reportTypeInfo, CsvStats csvStats)
         {
             // Find the stats by spinning through the graphs in this reporttype
             foreach (ReportGraph graph in reportTypeInfo.graphs)
@@ -1239,17 +1236,60 @@ namespace PerfSummaries
                     }
                 }
             }
-        }
+
+			base.PostInit(reportTypeInfo, csvStats);
+		}
 
         public List<string> sectionPrefixes;
 
-        List<double> multipliers;
-        List<double> budgets;
+		Dictionary<string, PeakStatInfo> statInfoLookup = new Dictionary<string, PeakStatInfo>();
+		class PeakStatInfo
+		{
+			public PeakStatInfo(string inName, string inShortName)
+			{
+				isInMainSummary = false;
+				multiplier = 1.0;
+				budget = 0.0;
+				name = inName;
+				shortName = inShortName;
+			}
+			public string name;
+			public string shortName;
+			public bool isInMainSummary;
+			public double multiplier;
+			public double budget;
+		};
 
-        List<string> shortenedStatNames;
-        List<string> hidePrefixes;
+		PeakStatInfo getOrAddStatInfo(string statName)
+		{
+			if ( statInfoLookup.ContainsKey(statName) )
+			{
+				return statInfoLookup[statName];
+			}
+			// Find the best (longest) prefix which matches this stat, and strip it off
+			int bestPrefixIndex = -1;
+			int bestPrefixLength = 0;
+			string shortStatName = statName;
+			for (int i = 0; i < hidePrefixes.Count; i++)
+			{
+				string prefix = hidePrefixes[i];
+				if (statName.ToLower().StartsWith(prefix) && prefix.Length > bestPrefixLength)
+				{
+					bestPrefixIndex = i;
+					bestPrefixLength = prefix.Length;
+				}
+			}
+			if (bestPrefixIndex >= 0)
+			{
+				shortStatName = statName.Substring(bestPrefixLength);
+			}
 
-        List<bool> isInMainSummary;
+			PeakStatInfo statInfo = new PeakStatInfo(statName,shortStatName);
+			statInfoLookup.Add(statName, statInfo);
+			return statInfo;
+		}
+
+        List<string> hidePrefixes;
     };
 
     class SummaryMetadataValue
