Fix issue with PerfReportTool not being able to find mono on Mac
#rnx

diff --git a/Engine/Source/Programs/CSVTools/PerfReportTool/PerfReportTool.cs b/Engine/Source/Programs/CSVTools/PerfReportTool/PerfReportTool.cs
--- a/Engine/Source/Programs/CSVTools/PerfReportTool/PerfReportTool.cs
+++ b/Engine/Source/Programs/CSVTools/PerfReportTool/PerfReportTool.cs
@@ -1191,13 +1191,14 @@ namespace PerfReportTool
 
 
 
-            string csvToolPath = GetBaseDirectory() + "\\CSVToSVG.exe";
+            string csvToolPath = GetBaseDirectory() + "/CSVToSVG.exe";
 			string binary = csvToolPath;
 
 			// run mono on non-Windows hosts
 			if (Host != HostPlatform.Windows)
 			{
-				binary = "mono";
+				// note, on Mac mono will not be on path
+				binary = Host == HostPlatform.Linux ? "mono" : "/Library/Frameworks/Mono.framework/Versions/Current/Commands/mono";
 				args = csvToolPath + " " + args;
 			}
             
