PerfReportTool 3.85:
Bulk query improvements:
- Add -beginEvent and -endEvent arguments to generate perf reports only within particular events (e.g skydiving begin/skydiving end)
- Add -summaryTableFilename arg so we can override the summary table filename

diff --git a/Engine/Source/Programs/CSVTools/CsvStats/CommandLineTool.cs b/Engine/Source/Programs/CSVTools/CsvStats/CommandLineTool.cs
--- a/Engine/Source/Programs/CSVTools/CsvStats/CommandLineTool.cs
+++ b/Engine/Source/Programs/CSVTools/CsvStats/CommandLineTool.cs
@@ -76,7 +76,18 @@ namespace CSVStats
             return CommandLineArgs.ContainsKey(key.ToLower());
         }
 
-        protected string GetArg(string key, bool mandatory = false)
+		protected string GetArg(string key, string defaultValue)
+		{
+			string lowerKey = key.ToLower();
+
+			if (CommandLineArgs.ContainsKey(lowerKey))
+			{
+				return CommandLineArgs[lowerKey];
+			}
+			return defaultValue;
+		}
+
+		protected string GetArg(string key, bool mandatory = false)
         {
             string lowerKey = key.ToLower();
 
@@ -92,7 +103,7 @@ namespace CSVStats
             return "";
         }
 
-        protected void WriteLine(String message, params object[] args)
+		protected void WriteLine(String message, params object[] args)
         {
             String formatted = String.Format(message, args);
             Console.WriteLine(formatted);
diff --git a/Engine/Source/Programs/CSVTools/CsvStats/CsvStats.cs b/Engine/Source/Programs/CSVTools/CsvStats/CsvStats.cs
--- a/Engine/Source/Programs/CSVTools/CsvStats/CsvStats.cs
+++ b/Engine/Source/Programs/CSVTools/CsvStats/CsvStats.cs
@@ -836,15 +836,16 @@ namespace CSVStats
             }
         }
 
-        public int StripByEvents(string startString, string endString)
+        public int StripByEvents(string startString, string endString, bool invert=false)
         {
             List<int> startIndices = null;
             List<int> endIndices = null;
             GetEventFrameIndexDelimiters(startString, endString, out startIndices, out endIndices);
-            if (startIndices.Count == 0 )
+			if (startIndices.Count == 0 )
             {
                 return 0;
             }
+
             int framesStripped = -1;
             int frameCount = 0;
             // Strip out samples and recompute averages
diff --git a/Engine/Source/Programs/CSVTools/PerfReportTool/PerfReportTool.cs b/Engine/Source/Programs/CSVTools/PerfReportTool/PerfReportTool.cs
--- a/Engine/Source/Programs/CSVTools/PerfReportTool/PerfReportTool.cs
+++ b/Engine/Source/Programs/CSVTools/PerfReportTool/PerfReportTool.cs
@@ -19,7 +19,7 @@ namespace PerfReportTool
 {
     class Version
     {
-        private static string VersionString = "3.84";
+        private static string VersionString = "3.85";
 
         public static string Get() { return VersionString; }
     };
@@ -260,7 +260,9 @@ namespace PerfReportTool
             "       -noStripEvents : if specified, don't strip out excluded events from the stats\n" +
 			"       -perfLog : output performance logging information\n" +
 			"       -writeSummaryCsv : if specified, a csv file containing summary information will be generated. Not available in bulk mode.\n" +
-            "\n" +
+			"       -beginEvent <event> : strip data before this event\n"+
+			"       -endEvent <event> : strip data after this event\n" +
+			"\n" +
 			"Optional bulk mode args: (use with -csvdir)\n" +
 			"       -recurse \n" +
 			"       -searchpattern <pattern>, e.g -searchpattern csvprofile*\n" +
@@ -273,6 +275,7 @@ namespace PerfReportTool
 			"       -summaryTable <name> :\n" +
 			"           Selects a custom summary table type from the list in reportTypes.xml \n"+
 			"           (if not specified, 'default' will be used)\n" +
+			"       -summaryTableFilename <name> : use the specified filename for the summary table (instead of SummaryTable.html)\n"+
             "       -metadataFilter <key=value,key=value...> : filters based on CSV metadata\n" +
             "       -precacheCount <n> : number of CSV files to precache in the lookahead cache (0 for no precache)\n" +
             "       -precacheThreadCount <n> : number of threads to use for the CSV lookahead cache (default 8)\n" +
@@ -449,6 +452,11 @@ namespace PerfReportTool
             // Write out the metadata table, if there is one
             if (metadataTable.Count > 0)
 			{
+				string summaryTableFilename = GetArg("summaryTableFilename", "SummaryTable");
+				if ( summaryTableFilename.ToLower().EndsWith(".html"))
+				{
+					summaryTableFilename = summaryTableFilename.Substring(0, summaryTableFilename.Length - 5);
+				}
 				string customSummaryTableFilter = GetArg("customTable");
 				bool bCsvTable = GetBoolArg("csvTable");
 				bool bCollateTable = GetBoolArg("collateTable");
@@ -459,11 +467,11 @@ namespace PerfReportTool
 					{
 						customSummaryTableRowSort = "buildversion,deviceprofile";
 					}
-					WriteMetadataTableReport(outputDir, "SummaryTable", metadataTable, customSummaryTableFilter.Split(',').ToList(), customSummaryTableRowSort.Split(',').ToList(), false, bCsvTable);
+					WriteMetadataTableReport(outputDir, summaryTableFilename, metadataTable, customSummaryTableFilter.Split(',').ToList(), customSummaryTableRowSort.Split(',').ToList(), false, bCsvTable);
 
 					if (bCollateTable)
 					{
-						WriteMetadataTableReport(outputDir, "SummaryTable_Collated", metadataTable, customSummaryTableFilter.Split(',').ToList(), customSummaryTableRowSort.Split(',').ToList(), true, bCsvTable);
+						WriteMetadataTableReport(outputDir, summaryTableFilename+"_Collated", metadataTable, customSummaryTableFilter.Split(',').ToList(), customSummaryTableRowSort.Split(',').ToList(), true, bCsvTable);
 					}
 				}
 				else
@@ -474,10 +482,10 @@ namespace PerfReportTool
 						summaryTableName = "default";
 					}
 					SummaryTableInfo tableInfo = reportXML.GetSummaryTable(summaryTableName);
-					WriteMetadataTableReport(outputDir, "SummaryTable", metadataTable, tableInfo, false, bCsvTable);
+					WriteMetadataTableReport(outputDir, summaryTableFilename, metadataTable, tableInfo, false, bCsvTable);
 					if (bCollateTable)
 					{
-						WriteMetadataTableReport(outputDir, "SummaryTable", metadataTable, tableInfo, true, bCsvTable);
+						WriteMetadataTableReport(outputDir, summaryTableFilename, metadataTable, tableInfo, true, bCsvTable);
 					}
 				}
 
@@ -485,7 +493,7 @@ namespace PerfReportTool
 				if (GetBoolArg("emailSummary") || GetBoolArg("emailTable"))
 				{
 					SummaryTableInfo tableInfo = reportXML.GetSummaryTable("condensed");
-					WriteMetadataTableReport(outputDir, "SummaryTable_Email", metadataTable, tableInfo, true);
+					WriteMetadataTableReport(outputDir, summaryTableFilename+"_Email", metadataTable, tableInfo, true);
 				}
                 perfLog.LogTiming("WriteSummaryTable");
             }
@@ -823,6 +831,33 @@ namespace PerfReportTool
 			}
 			csvStats.metaData.Values.Add("csvfilename", csvFile.filename);
 
+			// Adjust min/max x based on the event delimiters
+			string beginEventStr = GetArg("beginEvent").ToLower();
+			if (beginEventStr != "")
+			{
+				foreach (CsvEvent ev in csvStats.Events)
+				{
+					if (ev.Name.ToLower() == beginEventStr)
+					{
+						minX = Math.Max(minX, ev.Frame);
+						break;
+					}
+				}
+			}
+			string endEventStr = GetArg("endEvent").ToLower();
+			if ( endEventStr != "")
+			{ 
+				for (int i = csvStats.Events.Count - 1; i >= 0; i--)
+				{
+					CsvEvent ev = csvStats.Events[i];
+					if (ev.Name.ToLower() == endEventStr)
+					{
+						maxX = Math.Min(maxX, ev.Frame);
+						break;
+					}
+				}
+			}
+
 			// Crop the stats to the range
 			csvStats.CropStats(minX, maxX);
 
