CsvCollate 1.25
Averaging no longers shortens the output CSV to the minimum CSV length. Instead, stat frame values are averaged across the CSVs we have frames for. This allows for meaningful comparisons of large sets of CSVs
TODO: process the CSVs one at a time to reduce peak memory use

diff --git a/Engine/Source/Programs/CSVTools/CSVCollate/CSVCollate.cs b/Engine/Source/Programs/CSVTools/CSVCollate/CSVCollate.cs
--- a/Engine/Source/Programs/CSVTools/CSVCollate/CSVCollate.cs
+++ b/Engine/Source/Programs/CSVTools/CSVCollate/CSVCollate.cs
@@ -14,7 +14,7 @@ namespace CSVTools
 {
     class Version
     {
-        private static string VersionString = "1.24";
+        private static string VersionString = "1.25";
 
         public static string Get() { return VersionString; }
     };
diff --git a/Engine/Source/Programs/CSVTools/CsvStats/CsvStats.cs b/Engine/Source/Programs/CSVTools/CsvStats/CsvStats.cs
--- a/Engine/Source/Programs/CSVTools/CsvStats/CsvStats.cs
+++ b/Engine/Source/Programs/CSVTools/CsvStats/CsvStats.cs
@@ -661,22 +661,22 @@ namespace CSVStats
             Stats = newStats;
         }
 
-        // This will average only those stats that exist in all the stat sets,
-        // and it will clamp the number of the averaged stats to the minimum
-        // number across the all set.
+        // This will average only those stats that exist in all the stat sets
+		// Stats are averaged per frame, and the the final length is be equal to the longest CSV
+		// If CSVs are of varying length, this means later frames will be averaged over fewer samples than earlier frames
         public static CsvStats AverageByFrame(CsvStats[] statsToAvg, bool bStatsAvarage = false)
         {
             CsvStats avgStats = new CsvStats();
             if (statsToAvg.Length > 0)
             {
                 // We need to have all the frame per each stat in the file
-                int maxFrames = int.MaxValue;
+                int maxFrames = 0;
 
                 // Use the first as stat name basis
                 string[] statKeys = statsToAvg[0].Stats.Keys.ToArray();
                 foreach (string statName in statKeys)
                 {
-                    int minSamples = int.MaxValue;
+                    int maxSamples = 0;
                     int statCount = 0;
                     foreach (CsvStats stats in statsToAvg)
                     {
@@ -684,15 +684,15 @@ namespace CSVStats
                         // it doesn't exist in one of the set.
                         if (stats.Stats.ContainsKey(statName))
                         {
-                            minSamples = Math.Min(stats.Stats[statName].samples.Count, minSamples);
+                            maxSamples = Math.Max(stats.Stats[statName].samples.Count, maxSamples);
                             statCount++;
                         }
                     }
 
-                    if (statCount == statsToAvg.Length && minSamples < int.MaxValue && minSamples > 0)
+                    if (statCount == statsToAvg.Length && maxSamples > 0)
                     {
                         avgStats.AddStat(new StatSamples(statName));
-                        maxFrames = Math.Min(maxFrames, minSamples);
+                        maxFrames = Math.Max(maxFrames, maxSamples);
                     }
                 }
 
@@ -712,27 +712,31 @@ namespace CSVStats
                     StatSamples avgSamples = avgStats.GetStat(statName);
                     if (avgSamples != null)
                     {
-                        // Initialise sample to 0.0
-                        avgSamples.samples.AddRange(Enumerable.Repeat(0.0f, maxFrames));
+						List<int> sampleCounts = new List<int>();
+						sampleCounts.AddRange(Enumerable.Repeat(0, maxFrames));
+
+						// Initialise sample to 0.0
+						avgSamples.samples.AddRange(Enumerable.Repeat(0.0f, maxFrames));
 
                         // Add samples from other stats
                         foreach (CsvStats stats in statsToAvg)
                         {
                             StatSamples statSamples = stats.GetStat(statName);
-                            if ((statSamples != null) && (avgSamples.samples.Count <= statSamples.samples.Count))
+                            if ((statSamples != null) && (avgSamples.samples.Count >= statSamples.samples.Count))
                             {
-                                // This should always be true (avgSamples.samples.Count <= statSamples.samples.Count)
-                                for (var sIt = 0; sIt < avgSamples.samples.Count; ++sIt)
-                                {
-                                    avgSamples.samples[sIt] += statSamples.samples[sIt];
+								// This should always be true: avgSamples.samples.Count >= statSamples.samples.Count
+								for (int i = 0; i < statSamples.samples.Count; i++)
+								{
+									avgSamples.samples[i] += statSamples.samples[i];
+									sampleCounts[i] += 1;
                                 }
                             }
                         }
 
                         // Average the samples
-                        for (var sIt = 0; sIt < avgSamples.samples.Count; ++sIt)
+                        for (int i = 0; i < avgSamples.samples.Count; i++)
                         {
-                            avgSamples.samples[sIt] /= statsToAvg.Length;
+                            avgSamples.samples[i] /= (float)sampleCounts[i];
                         }
                     }
 
