PerfReportTool 3.86
Fix a bug with event names in event summaries getting truncated by 1 character in some cases

diff --git a/Engine/Source/Programs/CSVTools/PerfReportTool/PerfReportTool.cs b/Engine/Source/Programs/CSVTools/PerfReportTool/PerfReportTool.cs
--- a/Engine/Source/Programs/CSVTools/PerfReportTool/PerfReportTool.cs
+++ b/Engine/Source/Programs/CSVTools/PerfReportTool/PerfReportTool.cs
@@ -19,7 +19,7 @@ namespace PerfReportTool
 {
     class Version
     {
-        private static string VersionString = "3.85";
+        private static string VersionString = "3.86";
 
         public static string Get() { return VersionString; }
     };
diff --git a/Engine/Source/Programs/CSVTools/PerfReportTool/perfsummaries.cs b/Engine/Source/Programs/CSVTools/PerfReportTool/perfsummaries.cs
--- a/Engine/Source/Programs/CSVTools/PerfReportTool/perfsummaries.cs
+++ b/Engine/Source/Programs/CSVTools/PerfReportTool/perfsummaries.cs
@@ -734,7 +734,12 @@ namespace PerfSummaries
                 {
                     if (CsvStats.DoesSearchStringMatch(ev.Name, eventName))
                     {
-                        string eventContent = ev.Name.Substring(eventName.Length);
+						int len = eventName.Length;
+						if ( eventName.EndsWith("*"))
+						{
+							len--;
+						}
+						string eventContent = ev.Name.Substring(len).Trim();
                         if ( eventCountsDict.ContainsKey(eventContent))
                         {
                             eventCountsDict[eventContent]++;
