CsvCollate update - Process one CSV at a time when averaging (to avoid O(N^2) memory use, and output a stat reporting number of CSVs per frame.

diff --git a/Engine/Source/Programs/CSVTools/CSVCollate/CSVCollate.cs b/Engine/Source/Programs/CSVTools/CSVCollate/CSVCollate.cs
--- a/Engine/Source/Programs/CSVTools/CSVCollate/CSVCollate.cs
+++ b/Engine/Source/Programs/CSVTools/CSVCollate/CSVCollate.cs
@@ -14,7 +14,7 @@ namespace CSVTools
 {
     class Version
     {
-        private static string VersionString = "1.25";
+        private static string VersionString = "1.26";
 
         public static string Get() { return VersionString; }
     };
@@ -24,10 +24,13 @@ namespace CSVTools
         static string formatString =
             "Format: \n" +
             "       -csv <filename> OR -csvDir <path>\n" +
-            "       [-avg] <stats will be per frame averaged>\n" +
-            "       -o <csvFilename> \n";
+            "       [-avg] -stats will be per frame averaged\n" +
+			"       [-filterOutlierStat <stat>] - discard CSVs if this stat has very high values\n" +
+			"       [-filterOutlierThreshold <value>] - threshold for outliers (default:1000)\n" +
+			"       -o <csvFilename> \n";
 
-        void Run(string[] args)
+
+		void Run(string[] args)
         {
             // Read the command line
             if (args.Length < 1)
@@ -57,8 +60,11 @@ namespace CSVTools
             c = char.ToUpper(c);
             title = c + title.Substring(1);
 
-            // Whether or not we want stats to be averaged rather than appended
-            bool bAverage = GetBoolArg("avg");
+			string filterOutlierStat = GetArg("filterOutlierStat", false);
+			float filterOutlierThreshold = GetFloatArg("filterOutlierThreshold", 1000.0f);
+
+			// Whether or not we want stats to be averaged rather than appended
+			bool bAverage = GetBoolArg("avg");
 
             // Read CSV filenames from a directory or list
             string[] csvFilenames;
@@ -88,79 +94,117 @@ namespace CSVTools
 
             CsvStats combinedCsvStats = new CsvStats();
 
-            // List of stats to be averaged
-            CsvStats[] statsToAvg = new CsvStats[csvFilenames.Length];
+			// List of stats to be averaged
+			//CsvStats[] statsToAvg = new CsvStats[csvFilenames.Length];
 
-            // Read all the CSVs into one big CSVStats class
-            int csvIndex = 0;
+			// Read all the CSVs into one big CSVStats class
+			List<int> frameCsvCounts=new List<int>();
+			List<string> allCsvFilenames = new List<string>();
+			int csvIndex = 0;
             foreach (string csvFilename in csvFilenames)
             {
-                CsvStats csvStats = CsvStats.ReadCSVFile(csvFilename, null);
-
-                // Add the CSV filename as the first event (sanitised for CSV)
-                CsvEvent firstEvent = null;
-                firstEvent = new CsvEvent();
-                firstEvent.Frame = 0;
-                firstEvent.Name = "CSV:" + MakeShortFilename(csvFilename).Replace(' ', '_').Replace(',', '_').Replace('\n', '_');
-                csvStats.Events.Insert(0, firstEvent);
-
-                if (bAverage)
-                {
-                    statsToAvg[csvIndex] = csvStats;
-                }
-                else
-                {
-                    // Combine the stats
-                    if (csvIndex == 0)
-                    {
-                        combinedCsvStats = csvStats;
-                    }
-                    else
-                    {
-                        CsvMetadata metadataA = combinedCsvStats.metaData;
-                        CsvMetadata metadataB = csvStats.metaData;
-
-                        // If there is metadata, it should match
-                        if ( metadataA != null || metadataB != null)
-                        {
-                            if (!CsvMetadata.Matches(metadataA, metadataB))
-                            {
-                                WriteLine("Metadata mismatch!");
-                            }
-                        }
-                        combinedCsvStats.Combine(csvStats, false);
-                    }
-                }
-
-                csvIndex++;
-            }
-
-            if (bAverage)
-            {
-                // Average stats by frame
-                combinedCsvStats = CsvStats.AverageByFrame(statsToAvg);
-
-                if (combinedCsvStats.metaData != null)
-                {
-                  // Add few more meta information
-                  combinedCsvStats.metaData.Values.Add("Averaged", statsToAvg.Length.ToString());
-                }
-
-                int fIt = 0;
-                string[] filenames = new string[csvFilenames.Length];
-                foreach (var filename in csvFilenames)
-                {
-                    filenames[fIt++] = Path.GetFileName(filename);
-                }
-
-                if (combinedCsvStats.metaData != null)
-                {
-                  combinedCsvStats.metaData.Values.Add("Files", string.Join(",", filenames));
-                }
+                CsvStats srcCsvStats = CsvStats.ReadCSVFile(csvFilename, null);
+
+				// Check for outliers
+				bool skip = false;
+				if ( filterOutlierStat != null )
+				{
+					StatSamples outlierStat = srcCsvStats.GetStat(filterOutlierStat);
+					if ( outlierStat != null )
+					{
+						foreach (float sample in outlierStat.samples)
+						{
+							if (sample>filterOutlierThreshold)
+							{
+								WriteLine("CSV " + csvFilename + " ignored due to bad " + filterOutlierStat + " value: " + sample);
+								skip = true;
+								break;
+							}
+						}
+					}
+				}
+				if ( skip )
+				{
+					continue;
+				}
+
+				// Add the CSV filename as the first event if we're not averaging
+				if (!bAverage)
+				{
+					CsvEvent firstEvent = new CsvEvent();
+					firstEvent.Frame = 0;
+					firstEvent.Name = "CSV:" + MakeShortFilename(csvFilename).Replace(' ', '_').Replace(',', '_').Replace('\n', '_');
+					srcCsvStats.Events.Insert(0, firstEvent);
+				}
+
+				// Combine the stats
+				if (csvIndex == 0)
+				{
+					combinedCsvStats = srcCsvStats;
+				}
+				else
+				{
+					CsvMetadata metadataA = combinedCsvStats.metaData;
+					CsvMetadata metadataB = srcCsvStats.metaData;
+
+					// If there is metadata, it should match
+					if (metadataA != null || metadataB != null)
+					{
+						metadataA.CombineAndValidate(metadataB);
+					}
+					combinedCsvStats.Combine(srcCsvStats, bAverage, false);
+					if ( bAverage )
+					{
+						// Resize the framecount array to match
+						for (int i= frameCsvCounts.Count; i<combinedCsvStats.SampleCount; i++)
+						{
+							frameCsvCounts.Add(0);
+						}
+						for (int i=0; i<srcCsvStats.SampleCount;i++)
+						{
+							frameCsvCounts[i] += 1;
+						}
+					}
+				}
+				allCsvFilenames.Add(Path.GetFileName(csvFilename));
+				WriteLine("Csvs Processed: " + csvIndex + " / " + csvFilenames.Length);
+				csvIndex++;
             }
 
-            // Write the csv stats to a CSV
-            combinedCsvStats.WriteToCSV(csvOutFilename);
+			if (bAverage)
+			{
+				// Divide all samples by the total number of CSVs 
+				foreach (StatSamples stat in combinedCsvStats.Stats.Values)
+				{
+					for (int i=0; i<stat.samples.Count;i++)
+					{
+						stat.samples[i] /= (float)(frameCsvCounts[i]);
+					}
+				}
+
+				// Add a stat for the csv count
+				string csvCountStatName = "csvCount";
+				if (!combinedCsvStats.Stats.ContainsKey(csvCountStatName))
+				{
+					StatSamples csvCountStat = new StatSamples(csvCountStatName);
+					foreach (int count in frameCsvCounts)
+					{
+						csvCountStat.samples.Add((int)count);
+					}
+					combinedCsvStats.Stats.Add(csvCountStatName, csvCountStat);
+				}
+				if (combinedCsvStats.metaData != null)
+				{
+					// Add some metadata
+					combinedCsvStats.metaData.Values.Add("Averaged", allCsvFilenames.Count.ToString());
+					combinedCsvStats.metaData.Values.Add("SourceFiles", string.Join(";", allCsvFilenames));
+				}
+			}
+
+			combinedCsvStats.ComputeAveragesAndTotal();
+
+			// Write the csv stats to a CSV
+			combinedCsvStats.WriteToCSV(csvOutFilename);
         }
 
 
diff --git a/Engine/Source/Programs/CSVTools/CsvStats/CsvStats.cs b/Engine/Source/Programs/CSVTools/CsvStats/CsvStats.cs
--- a/Engine/Source/Programs/CSVTools/CsvStats/CsvStats.cs
+++ b/Engine/Source/Programs/CSVTools/CsvStats/CsvStats.cs
@@ -134,7 +134,7 @@ namespace CSVStats
 
             foreach (string key in valuesDontMatchKeys)
             {
-                Values[key] = "["+ key + "s don't match!]";
+                Values[key] = "["+ key + " doesn't match!]";
             }
         }
 
@@ -618,7 +618,7 @@ namespace CSVStats
             }
         }
 
-        public void Combine(CsvStats srcStats, bool computeAverage = true)
+        public void Combine(CsvStats srcStats, bool sumRows = false, bool computeAverage = true)
         {
             if (metaData != null)
             {
@@ -626,18 +626,35 @@ namespace CSVStats
             }
 
             // Collate the stats, removing ones which don't exist in both
-            int FrameOffset = 0;
+            int EventFrameOffset = 0;
             Dictionary<string, StatSamples> newStats = new Dictionary<string, StatSamples>();
             foreach (StatSamples srcStat in srcStats.Stats.Values.ToArray())
             {
                 if (Stats.ContainsKey(srcStat.Name.ToLower()))
                 {
                     StatSamples destStat = GetStat(srcStat.Name);
-                    FrameOffset = Math.Max(FrameOffset, destStat.samples.Count);
-                    foreach (float sample in srcStat.samples)
-                    {
-                        destStat.samples.Add(sample);
-                    };
+					if (sumRows)
+					{
+						int n = Math.Min(srcStat.samples.Count, destStat.samples.Count);
+						// Combine initial stats
+						for (int i=0;i<n; i++)
+						{
+							destStat.samples[i] += srcStat.samples[i];
+						}
+						// Add the remaining stats
+						for (int i=n; i< srcStat.samples.Count; i++)
+						{
+							destStat.samples.Add(srcStat.samples[i]);
+						}
+					}
+					else
+					{
+						EventFrameOffset = Math.Max(EventFrameOffset, destStat.samples.Count);
+						foreach (float sample in srcStat.samples)
+						{
+							destStat.samples.Add(sample);
+						};
+					}
                     newStats.Add(srcStat.Name.ToLower(), destStat);
                 }
             }
@@ -651,20 +668,29 @@ namespace CSVStats
                 }
             }
 
-            // Add the events, offsetting the frame numbers
-            foreach (CsvEvent ev in srcStats.Events)
-            {
-                ev.Frame += FrameOffset;
-                Events.Add(ev);
-            }
-
+			if ( sumRows)
+			{
+				// just remove events. Averaging is complicated/doesn't really make sense
+				Events.Clear();
+			}
+			else
+			{
+				// Add the events, offsetting the frame numbers
+				foreach (CsvEvent ev in srcStats.Events)
+				{
+					ev.Frame += EventFrameOffset;
+					Events.Add(ev);
+				}
+			}
+			
             Stats = newStats;
         }
 
-        // This will average only those stats that exist in all the stat sets
+
+		// This will average only those stats that exist in all the stat sets
 		// Stats are averaged per frame, and the the final length is be equal to the longest CSV
 		// If CSVs are of varying length, this means later frames will be averaged over fewer samples than earlier frames
-        public static CsvStats AverageByFrame(CsvStats[] statsToAvg, bool bStatsAvarage = false)
+		public static CsvStats AverageByFrame(CsvStats[] statsToAvg, bool bStatsAvarage = false)
         {
             CsvStats avgStats = new CsvStats();
             if (statsToAvg.Length > 0)
