PerfReportTool
- PerfReportTool will now collect stats from all summaries in reports
- Fix peak sectionPrefix option not working due to case-sensitive comparison.

diff --git a/Engine/Source/Programs/CSVTools/PerfReportTool/PerfReportTool.cs b/Engine/Source/Programs/CSVTools/PerfReportTool/PerfReportTool.cs
--- a/Engine/Source/Programs/CSVTools/PerfReportTool/PerfReportTool.cs
+++ b/Engine/Source/Programs/CSVTools/PerfReportTool/PerfReportTool.cs
@@ -19,7 +19,7 @@ namespace PerfReportTool
 {
     class Version
     {
-        private static string VersionString = "3.86";
+        private static string VersionString = "3.87";
 
         public static string Get() { return VersionString; }
     };
@@ -670,8 +670,11 @@ namespace PerfReportTool
             List<string> statsToSummariseList = null;
             if ( reportTypeInfo.summaries.Count >0 )
             {
-                // TODO: multiple summaries
-                statsToSummariseList = reportTypeInfo.summaries[0].GetUniqueStatNames().ToList();
+				statsToSummariseList = new List<string>();
+				foreach (var summary in reportTypeInfo.summaries)
+				{
+					statsToSummariseList.AddRange(summary.GetUniqueStatNames());
+				}
                 for (int i = 0; i < statsToSummariseList.Count; i++)
                 {
                     // Remove all Min/max sort tags before checking stats against existing stats.
diff --git a/Engine/Source/Programs/CSVTools/PerfReportTool/perfsummaries.cs b/Engine/Source/Programs/CSVTools/PerfReportTool/perfsummaries.cs
--- a/Engine/Source/Programs/CSVTools/PerfReportTool/perfsummaries.cs
+++ b/Engine/Source/Programs/CSVTools/PerfReportTool/perfsummaries.cs
@@ -1129,7 +1129,7 @@ namespace PerfSummaries
 				PeakStatInfo statInfo = getOrAddStatInfo(stat);
 
 				if ((csvStats.Stats.ContainsKey(statName.ToLower())) && // If the main stats table contains this stat AND
-                     (sectionPrefix == null || statName.StartsWith(sectionPrefix)) // If there is no hide prefix just display the stat, otherwise, make sure they match.
+                     (sectionPrefix == null || statName.StartsWith(sectionPrefix, true, null)) // If there is no hide prefix just display the stat, otherwise, make sure they match.
                    )
                 {
                     // Do the calculations for the averages and peak, and then write it to the table along with the budget.
