CsvCollate 1.30
- Add support for search patterns via -searchpattern <mypattern>
- Add -recurse arg to make recursion optional
- List CSV files to log
- Fix out-by-one error with logging

diff --git a/Engine/Source/Programs/CSVTools/CSVCollate/CSVCollate.cs b/Engine/Source/Programs/CSVTools/CSVCollate/CSVCollate.cs
--- a/Engine/Source/Programs/CSVTools/CSVCollate/CSVCollate.cs
+++ b/Engine/Source/Programs/CSVTools/CSVCollate/CSVCollate.cs
@@ -14,7 +14,7 @@ namespace CSVTools
 {
     class Version
     {
-        private static string VersionString = "1.26";
+        private static string VersionString = "1.30";
 
         public static string Get() { return VersionString; }
     };
@@ -23,8 +23,10 @@ namespace CSVTools
     {
         static string formatString =
             "Format: \n" +
-            "       -csv <filename> OR -csvDir <path>\n" +
-            "       [-avg] -stats will be per frame averaged\n" +
+            "       -csv <filename>  -csvDir <path>\n" +
+			"       [-searchPattern <pattern, e.g *.csv>] - for use with -csvDir\n" +
+			"       [-avg] -stats will be per frame averaged\n" +
+			"       [-recurse] - for use with -csvdir\n" +
 			"       [-filterOutlierStat <stat>] - discard CSVs if this stat has very high values\n" +
 			"       [-filterOutlierThreshold <value>] - threshold for outliers (default:1000)\n" +
 			"       -o <csvFilename> \n";
@@ -72,7 +74,9 @@ namespace CSVTools
             if (csvDir.Length > 0)
             {
                 DirectoryInfo di = new DirectoryInfo(csvDir);
-                var files = di.GetFiles("*.csv", SearchOption.AllDirectories);
+				string searchPattern = GetArg("searchPattern","*.csv");
+				bool bRecurse = GetBoolArg("recurse");
+				var files = di.GetFiles(searchPattern, bRecurse ? SearchOption.AllDirectories : SearchOption.TopDirectoryOnly);
                 csvFilenames = new string[files.Length];
                 int i = 0;
                 foreach (FileInfo csvFile in files)
@@ -92,12 +96,16 @@ namespace CSVTools
                 csvFilenames = csvFilenamesStr.Split(';');
             }
 
-            CsvStats combinedCsvStats = new CsvStats();
+			Console.WriteLine("Collating " + csvFilenames.Length + " csvs:");
+			foreach (string csvFilename in csvFilenames)
+			{
+				Console.WriteLine("  "+csvFilename);
+			}
+			Console.WriteLine("");
+
+			CsvStats combinedCsvStats = new CsvStats();
 
-			// List of stats to be averaged
-			//CsvStats[] statsToAvg = new CsvStats[csvFilenames.Length];
 
-			// Read all the CSVs into one big CSVStats class
 			List<int> frameCsvCounts=new List<int>();
 			List<string> allCsvFilenames = new List<string>();
 			int csvIndex = 0;
@@ -167,9 +175,9 @@ namespace CSVTools
 					}
 				}
 				allCsvFilenames.Add(Path.GetFileName(csvFilename));
-				WriteLine("Csvs Processed: " + csvIndex + " / " + csvFilenames.Length);
 				csvIndex++;
-            }
+				WriteLine("Csvs Processed: " + csvIndex + " / " + csvFilenames.Length);
+			}
 
 			if (bAverage)
 			{
