PerfReportTool 3.88
- Fix issue when adding numeric and non-numeric metadata to summary metadata columns. When this happens, we just convert the column to strings instead of throwing an exception

diff --git a/Engine/Source/Programs/CSVTools/PerfReportTool/PerfReportTool.cs b/Engine/Source/Programs/CSVTools/PerfReportTool/PerfReportTool.cs
--- a/Engine/Source/Programs/CSVTools/PerfReportTool/PerfReportTool.cs
+++ b/Engine/Source/Programs/CSVTools/PerfReportTool/PerfReportTool.cs
@@ -19,7 +19,7 @@ namespace PerfReportTool
 {
     class Version
     {
-        private static string VersionString = "3.87";
+        private static string VersionString = "3.88";
 
         public static string Get() { return VersionString; }
     };
diff --git a/Engine/Source/Programs/CSVTools/PerfReportTool/perfsummaries.cs b/Engine/Source/Programs/CSVTools/PerfReportTool/perfsummaries.cs
--- a/Engine/Source/Programs/CSVTools/PerfReportTool/perfsummaries.cs
+++ b/Engine/Source/Programs/CSVTools/PerfReportTool/perfsummaries.cs
@@ -1430,7 +1430,9 @@ namespace PerfSummaries
 		{
 			if ( !isNumeric )
 			{
-				throw new Exception("can't mix numeric and non-numeric types!");
+				// This is already a non-numeric column. Better treat this as a string value
+				SetStringValue(index, value.ToString());
+				return;
 			}
 			// Grow to fill if necessary
 			if ( index >= floatValues.Count )
@@ -1443,6 +1445,20 @@ namespace PerfSummaries
 			floatValues[index] = value;
 		}
 
+		void convertToStrings()
+		{
+			if ( isNumeric )
+			{
+				stringValues = new List<string>();
+				foreach (float f in floatValues)
+				{
+					stringValues.Add(f.ToString());
+				}
+				floatValues = new List<float>();
+				isNumeric = false;
+			}
+		}
+
 		public void SetColourThresholds(int index, ColourThresholdList value)
 		{
 			// Grow to fill if necessary
@@ -1492,7 +1508,8 @@ namespace PerfSummaries
 		{
 			if (isNumeric)
 			{
-				throw new Exception("can't mix numeric and non-numeric types!");
+				// Better convert this to a string column, since we're trying to add a string to it
+				convertToStrings();
 			}
 			// Grow to fill if necessary
 			if (index >= stringValues.Count)
