PerfReportTool 3.90
- Add boundedstatvalues summary type for visualizing particular stats between two events, with configurable columns.
  - Each column displays a stat value, computed with one of the following formulas: sum, percentoverthreshold, percentunderthreshold, average and streamingstressmetric (the latter is used for File I/O)
- Add support for shared summaries (setup once and use in multiple reports)
- Add snapToPeaks property for graphs, so we can disable it in some cases (e.g for smooth graphs)

diff --git a/Engine/Source/Programs/CSVTools/PerfReportTool/PerfReportTool.cs b/Engine/Source/Programs/CSVTools/PerfReportTool/PerfReportTool.cs
--- a/Engine/Source/Programs/CSVTools/PerfReportTool/PerfReportTool.cs
+++ b/Engine/Source/Programs/CSVTools/PerfReportTool/PerfReportTool.cs
@@ -19,7 +19,7 @@ namespace PerfReportTool
 {
     class Version
     {
-        private static string VersionString = "3.88";
+        private static string VersionString = "3.90";
 
         public static string Get() { return VersionString; }
     };
@@ -748,6 +748,7 @@ namespace PerfReportTool
 
 			if (!GetBoolArg("readAllStats"))
 			{
+				// Strip out stats if they're not in the summary
 				List<string> ListOfKeys = new List<string>(csvStats.Stats.Keys);
 				for (int i = csvStats.Stats.Keys.Count - 1; i >= 0; i--)
 				{
@@ -986,7 +987,10 @@ namespace PerfReportTool
 				{
 					htmlFile.WriteLine("<ul>");
 					string currentCategory = sections[index];
-					htmlFile.WriteLine("<h4>" + currentCategory + " Graphs</h4>");
+					if (currentCategory.Length > 0)
+					{
+						htmlFile.WriteLine("<h4>" + currentCategory + " Graphs</h4>");
+					}
 					for (int i = 0; i < svgFilenames.Count(); i++)
 					{
 						string svgFilename = svgFilenames[i];
@@ -1130,8 +1134,13 @@ namespace PerfReportTool
             bool stacked = graphSettings.stacked.value;
             bool showAverages = graphSettings.showAverages.value;
             bool filterOutZeros = graphSettings.filterOutZeros.value;
+			bool snapToPeaks = false;
+			if (graphSettings.snapToPeaks.isSet)
+			{
+				snapToPeaks = graphSettings.snapToPeaks.value;
+			}
 
-            int maxHierarchyDepth = graphSettings.maxHierarchyDepth.value;
+			int maxHierarchyDepth = graphSettings.maxHierarchyDepth.value;
             string hideStatPrefix = graphSettings.hideStatPrefix.value;
             string showEvents = graphSettings.showEvents.value;
             double statMultiplier = graphSettings.statMultiplier.isSet ? graphSettings.statMultiplier.value : 1.0;
@@ -1189,7 +1198,7 @@ namespace PerfReportTool
 
 			string shortFileName = MakeShortFilename(csvFilename).Replace(" ", "_");
 
-			outFilename += "_" + shortFileName + svgName + ".svg";
+			outFilename += "_" + shortFileName + svgName.Replace("/", "_SLASH_").Replace("\\", "_BSLASH_") + ".svg";
             string args =
                 " -csvs \""  + csvFilename  + "\"" +
                 " -title \"" + title + "\"" +
@@ -1215,7 +1224,8 @@ namespace PerfReportTool
                 (interactive                  ? " -interactive" : "") +
                 (stacked                      ? " -stacked -forceLegendSort" : "") +
                 (showAverages                 ? " -showAverages" : "") +
-                (filterOutZeros               ? " -filterOutZeros" : "") +
+				(snapToPeaks                  ? "" : " -nosnap") +
+				(filterOutZeros               ? " -filterOutZeros" : "") +
                 (maxHierarchyDepth >=0        ? " -maxHierarchyDepth " + maxHierarchyDepth.ToString() : "") +
                 (hideStatPrefix.Length>0      ? " -hideStatPrefix " + hideStatPrefix : "") +
                 (graphSettings.mainStat.isSet ? " -stacktotalstat " + graphSettings.mainStat.value : "") +
@@ -1308,7 +1318,7 @@ namespace PerfReportTool
 
     class ReportTypeInfo
     {
-        public ReportTypeInfo(XElement element)
+        public ReportTypeInfo(XElement element, Dictionary<string,XElement> sharedSummaries)
         {
             graphs = new List<ReportGraph>();
             summaries = new List<Summary>();
@@ -1320,28 +1330,41 @@ namespace PerfReportTool
 					ReportGraph graph = new ReportGraph(child);
 					graphs.Add(graph);
 				}
-				else if (child.Name == "summary")
+				else if (child.Name == "summary" || child.Name=="summaryRef")
 				{
-					string summaryType = child.Attribute("type").Value;
+					XElement summaryElement = null;
+					if (child.Name == "summaryRef")
+					{
+						summaryElement = sharedSummaries[child.Attribute("name").Value];
+					}
+					else
+					{
+						summaryElement = child;
+					}
+					string summaryType = summaryElement.Attribute("type").Value;
 					if (summaryType == "histogram")
 					{
-						summaries.Add(new HistogramSummary(child));
+						summaries.Add(new HistogramSummary(summaryElement));
 					}
 					else if (summaryType == "peak")
 					{
-						summaries.Add(new PeakSummary(child));
+						summaries.Add(new PeakSummary(summaryElement));
 					}
 					else if (summaryType == "fpschart")
 					{
-						summaries.Add(new FPSChartSummary(child));
+						summaries.Add(new FPSChartSummary(summaryElement));
 					}
 					else if (summaryType == "hitches")
 					{
-						summaries.Add(new HitchSummary(child));
+						summaries.Add(new HitchSummary(summaryElement));
 					}
 					else if (summaryType == "event")
 					{
-						summaries.Add(new EventSummary(child));
+						summaries.Add(new EventSummary(summaryElement));
+					}
+					else if (summaryType == "boundedstatvalues")
+					{
+						summaries.Add(new BoundedStatValuesSummary(summaryElement));
 					}
 				}
 				else if (child.Name == "metadataToShow")
@@ -1410,6 +1433,7 @@ namespace PerfReportTool
             additionalArgs = new OptionalString(element, "additionalArgs", true);
             statMultiplier = new OptionalDouble(element, "statMultiplier");
 			legendAverageThreshold = new OptionalDouble(element, "legendAverageThreshold");
+			snapToPeaks = new OptionalBool(element, "snapToPeaks");
 		}
 		public void InheritFrom(GraphSettings baseSettings)
         {
@@ -1439,6 +1463,7 @@ namespace PerfReportTool
             statMultiplier.InheritFrom(baseSettings.statMultiplier);
             ignoreStats.InheritFrom(baseSettings.ignoreStats);
 			legendAverageThreshold.InheritFrom(baseSettings.legendAverageThreshold);
+			snapToPeaks.InheritFrom(baseSettings.snapToPeaks);
 
 		}
         public OptionalBool smooth;
@@ -1468,7 +1493,9 @@ namespace PerfReportTool
 		public OptionalDouble legendAverageThreshold;
 
 		public OptionalBool requiresDetailedStats;
-    };
+		public OptionalBool snapToPeaks;
+
+	};
 
 	static class Extensions
 	{
@@ -1730,6 +1757,17 @@ namespace PerfReportTool
 				}
 			}
 
+			// Add any shared summaries
+			XElement sharedSummariesElement = rootElement.Element("sharedSummaries");
+			sharedSummaries = new Dictionary<string, XElement>();
+			if (sharedSummariesElement != null)
+			{
+				foreach (XElement summaryElement in sharedSummariesElement.Elements("summary"))
+				{
+					sharedSummaries.Add(summaryElement.Attribute("refName").Value, summaryElement);
+				}
+			}
+
 		}
 
 		public ReportTypeInfo GetReportTypeInfo(string reportType, CachedCsvFile csvFile, bool bBulkMode, bool forceReportType )
@@ -1742,7 +1780,7 @@ namespace PerfReportTool
 				{
 					if (IsReportTypeXMLCompatibleWithStats(element, csvFile.dummyCsvStats))
 					{
-						reportTypeInfo = new ReportTypeInfo(element);
+						reportTypeInfo = new ReportTypeInfo(element, sharedSummaries);
 						break;
 					}
 				}
@@ -1753,20 +1791,20 @@ namespace PerfReportTool
 			}
 			else
 			{
-				XElement foundReporTypeElement = null;
+				XElement foundReportTypeElement = null;
 				foreach (XElement element in reportTypesElement.Elements("reporttype"))
 				{
 					if (element.Attribute("name").Value.ToLower() == reportType)
 					{
-						foundReporTypeElement = element;
+						foundReportTypeElement = element;
 					}
 				}
-				if (foundReporTypeElement == null)
+				if (foundReportTypeElement == null)
 				{
 					throw new Exception("Report type " + reportType + " not found in " + reportTypeXmlFilename);
 				}
 
-                if (!IsReportTypeXMLCompatibleWithStats(foundReporTypeElement, csvFile.dummyCsvStats))
+                if (!IsReportTypeXMLCompatibleWithStats(foundReportTypeElement, csvFile.dummyCsvStats))
                 {
                     if (forceReportType)
                     {
@@ -1777,7 +1815,7 @@ namespace PerfReportTool
                         throw new Exception("Report type " + reportType + " is not compatible with CSV " + csvFile.filename);
                     }
                 }
-                reportTypeInfo = new ReportTypeInfo(foundReporTypeElement);
+                reportTypeInfo = new ReportTypeInfo(foundReportTypeElement,sharedSummaries);
             }
 
             // Load the graphs
@@ -1901,6 +1939,7 @@ namespace PerfReportTool
 		XElement rootElement;
 		XElement graphGroupsElement;
 		XElement summaryTablesElement;
+		Dictionary<string,XElement> sharedSummaries;
 		Dictionary<string, GraphSettings> graphs;
 		Dictionary<string, string> statDisplayNameMapping;
 		DerivedMetadataMappings derivedMetadataMappings;
diff --git a/Engine/Source/Programs/CSVTools/PerfReportTool/perfsummaries.cs b/Engine/Source/Programs/CSVTools/PerfReportTool/perfsummaries.cs
--- a/Engine/Source/Programs/CSVTools/PerfReportTool/perfsummaries.cs
+++ b/Engine/Source/Programs/CSVTools/PerfReportTool/perfsummaries.cs
@@ -341,6 +341,25 @@ namespace PerfSummaries
             return uniqueStats.ToArray();
         }
 
+		protected double [] ReadColourThresholdsXML(XElement colourThresholdEl)
+		{
+			if (colourThresholdEl != null)
+			{
+				string[] colourStrings = colourThresholdEl.Value.Split(',');
+				if (colourStrings.Length != 4)
+				{
+					throw new Exception("Incorrect number of colourthreshold entries. Should be 4.");
+				}
+				double [] colourThresholds = new double[4];
+				for (int i = 0; i < colourStrings.Length; i++)
+				{
+					colourThresholds[i] = Convert.ToDouble(colourStrings[i], System.Globalization.CultureInfo.InvariantCulture);
+				}
+				return colourThresholds;
+			}
+			return null;
+		}
+
         public string GetStatThresholdColour(string StatToUse, double value)
         {
 			ColourThresholdList Thresholds = GetStatColourThresholdList(StatToUse);
@@ -705,23 +724,7 @@ namespace PerfSummaries
             title = element.GetSafeAttibute("title","Events");
             metadataKey = element.Attribute("metadataKey").Value;
             events = element.Element("events").Value.Split(',');
-
-
-            XElement colourThresholdEl = element.Element("colourThresholds");
-            if (colourThresholdEl != null)
-            {
-                string[] colourStrings = colourThresholdEl.Value.Split(',');
-                if (colourStrings.Length != 4)
-                {
-                    throw new Exception("Incorrect number of colourthreshold entries. Should be 4.");
-                }
-                colourThresholds = new double[4];
-                for (int i = 0; i < colourStrings.Length; i++)
-                {
-                    colourThresholds[i] = Convert.ToDouble(colourStrings[i], System.Globalization.CultureInfo.InvariantCulture);
-                }
-            }
-
+			colourThresholds = ReadColourThresholdsXML(element.Element("colourThresholds"));
         }
 
         public override void WriteSummaryData(System.IO.StreamWriter htmlFile, CsvStats csvStats, bool bIncludeSummaryCsv, SummaryMetadata metadata, string htmlFileName)
@@ -892,16 +895,8 @@ namespace PerfSummaries
         public HistogramSummary(XElement element)
         {
             ReadStatsFromXML(element);
-            string[] colourStrings = element.Element("colourThresholds").Value.Split(',');
-            if (colourStrings.Length != 4)
-            {
-                throw new Exception("Incorrect number of colourthreshold entries. Should be 4.");
-            }
-            ColourThresholds = new double[4];
-            for (int i = 0; i < colourStrings.Length; i++)
-            {
-                ColourThresholds[i] = Convert.ToDouble(colourStrings[i], System.Globalization.CultureInfo.InvariantCulture);
-            }
+
+			ColourThresholds = ReadColourThresholdsXML(element.Element("colourThresholds"));
 
             string[] histogramStrings = element.Element("histogramThresholds").Value.Split(',');
             HistogramThresholds = new double[histogramStrings.Length];
@@ -1297,7 +1292,277 @@ namespace PerfSummaries
         List<string> hidePrefixes;
     };
 
-    class SummaryMetadataValue
+	class BoundedStatValuesSummary : Summary
+	{
+		class Column
+		{
+			public string name;
+			public string formula;
+			public double value;
+			public string metadataKey;
+			public string statName;
+			public bool perSecond;
+			public bool filterOutZeros;
+			public bool applyEndOffset;
+			public double multiplier;
+			public double threshold;
+			public ColourThresholdList colourThresholdList;
+		};
+		public BoundedStatValuesSummary(XElement element)
+		{
+			ReadStatsFromXML(element);
+			if (stats.Count != 0)
+			{
+				throw new Exception("<stats> element is not supported");
+			}
+
+			title = element.GetSafeAttibute("title", "Events");
+			beginEvent = element.Attribute("beginevent").Value;
+			endEvent = element.Attribute("endevent").Value;
+
+			endOffsetPercentage = 0.0;
+			XAttribute endOffsetAtt = element.Attribute("endoffsetpercent");
+			if ( endOffsetAtt != null )
+			{
+				endOffsetPercentage = double.Parse(endOffsetAtt.Value);
+			}
+			columns = new List<Column>();
+
+			foreach (XElement columnEl in element.Elements("column"))
+			{
+				Column column = new Column();
+				double[] colourThresholds = ReadColourThresholdsXML(columnEl.Element("colourThresholds"));
+				if (colourThresholds != null)
+				{
+					column.colourThresholdList = new ColourThresholdList();
+					for (int i = 0; i < colourThresholds.Length; i++)
+					{
+						column.colourThresholdList.Add(new ThresholdInfo(colourThresholds[i], null));
+					}
+				}
+
+				XAttribute metadataKeyAtt = columnEl.Attribute("metadataKey");
+				if (metadataKeyAtt!=null)
+				{
+					column.metadataKey = metadataKeyAtt.Value;
+				}
+				column.statName = columnEl.Attribute("stat").Value.ToLower();
+				if ( !stats.Contains(column.statName) )
+				{
+					stats.Add(column.statName);
+				}
+
+				column.name = columnEl.Attribute("name").Value;
+				column.formula = columnEl.Attribute("formula").Value;
+				column.filterOutZeros= columnEl.GetSafeAttibute<bool>("filteroutzeros", false);
+				column.perSecond = columnEl.GetSafeAttibute<bool>("persecond", false);
+				column.multiplier = columnEl.GetSafeAttibute<double>("multiplier", 1.0);
+				column.threshold = columnEl.GetSafeAttibute<double>("threshold", 0.0);
+				column.applyEndOffset = columnEl.GetSafeAttibute<bool>("applyEndOffset", true);
+				columns.Add(column);
+			}
+		}
+
+		public override void WriteSummaryData(System.IO.StreamWriter htmlFile, CsvStats csvStats, bool bIncludeSummaryCsv, SummaryMetadata metadata, string htmlFileName)
+		{
+			int startFrame = -1;
+			int endFrame = int.MaxValue;
+
+			// Find the start and end frames based on the events
+			foreach (CsvEvent ev in csvStats.Events)
+			{
+				if (CsvStats.DoesSearchStringMatch(ev.Name, beginEvent))
+				{
+					startFrame = Math.Max(ev.Frame,startFrame);
+				}
+				if (CsvStats.DoesSearchStringMatch(ev.Name, endEvent))
+				{
+					endFrame = Math.Min(ev.Frame, endFrame);
+				}
+			}
+			if ( startFrame == -1 )
+			{
+				throw new Exception("BoundedStatValuesSummary: Begin event " + beginEvent+" was not found");
+			}
+			if (endFrame== -1)
+			{
+				throw new Exception("BoundedStatValuesSummary: End event " + endEvent + " was not found");
+			}
+			if ( startFrame >= endFrame )
+			{
+				throw new Exception("BoundedStatValuesSummary: end event appeared before the start event");
+			}
+
+			int endEventFrame = Math.Min(csvStats.SampleCount, endFrame + 1);
+			if (endOffsetPercentage > 0.0)
+			{
+				double multiplier = endOffsetPercentage / 100.0;
+				endFrame += (int)((double)(endFrame-startFrame)*multiplier);
+			}
+
+			endFrame = Math.Min(csvStats.SampleCount, endFrame + 1);
+			StatSamples frameTimeStat = csvStats.GetStat("frametime");
+			List<float> frameTimes = frameTimeStat.samples;
+
+			// Filter only columns with stats that exist in the CSV
+			List<Column> filteredColumns = new List<Column>();
+			foreach (Column col in columns)
+			{
+				if (csvStats.GetStat(col.statName) != null)
+				{
+					filteredColumns.Add(col);
+				}
+			}
+
+
+			// Process the column values
+			foreach (Column col in filteredColumns)
+			{
+				List<float> statValues = csvStats.GetStat(col.statName).samples;
+				double value = 0.0;
+				double totalFrameWeight = 0.0;
+				int colEndFrame = col.applyEndOffset ? endFrame : endEventFrame;
+
+				if ( col.formula == "average")
+				{
+					for (int i=startFrame; i< colEndFrame; i++)
+					{
+						if (col.filterOutZeros == false || statValues[i] > 0)
+						{
+							value += statValues[i] * frameTimes[i];
+							totalFrameWeight += frameTimes[i];
+						}
+					}
+				}
+				else if (col.formula == "percentoverthreshold")
+				{
+					for (int i = startFrame; i < colEndFrame; i++)
+					{
+						if (statValues[i] > col.threshold)
+						{
+							value += frameTimes[i];
+						}
+						totalFrameWeight += frameTimes[i];
+					}
+					value *= 100.0;
+				}
+				else if (col.formula == "percentunderthreshold")
+				{
+					for (int i = startFrame; i < colEndFrame; i++)
+					{
+						if (statValues[i] < col.threshold)
+						{
+							value += frameTimes[i];
+						}
+						totalFrameWeight += frameTimes[i];
+					}
+					value *= 100.0;
+				}
+				else if (col.formula == "sum")
+				{
+					for (int i = startFrame; i < colEndFrame; i++)
+					{
+						value += statValues[i];
+					}
+
+					if (col.perSecond)
+					{
+						double totalTimeMS = 0.0;
+						for (int i = startFrame; i < colEndFrame; i++)
+						{
+							if (col.filterOutZeros == false || statValues[i] > 0)
+							{
+								totalTimeMS += frameTimes[i];
+							}
+						}
+						value /= (totalTimeMS / 1000.0);
+					}
+					totalFrameWeight = 1.0;
+				}
+				else if (col.formula == "streamingstressmetric")
+				{
+					// Note: tInc is scaled such that it hits 1.0 on the event frame, regardless of the offset
+					double tInc = 1.0/(double)(endEventFrame - startFrame);
+					double t = tInc*0.5;
+					for (int i = startFrame; i < colEndFrame; i++)
+					{
+						if (col.filterOutZeros == false || statValues[i] > 0)
+						{
+							// Frame weighting is scaled to heavily favor final frames. Note that t can exceed 1 after the event frame if an offset percentage is specified, so we clamp it
+							double frameWeight = Math.Pow(Math.Min(t,1.0), 4.0) * frameTimes[i];
+
+							// If we're past the end event frame, apply a linear falloff to the weight
+							if (i >= endEventFrame)
+							{
+								double falloff = 1.0 - (double)(i - endEventFrame) / (colEndFrame - endEventFrame);
+								frameWeight *= falloff;
+							}
+
+							// The frame score takes into account the queue depth, but it's not massively significant
+							double frameScore = Math.Pow(statValues[i], 0.25);
+							value += frameScore * frameWeight;
+							totalFrameWeight += frameWeight;
+						}
+						t += tInc;
+					}
+				}
+				else
+				{
+					throw new Exception("BoundedStatValuesSummary: unexpected formula "+col.formula);
+				}
+				value *= col.multiplier;
+				col.value = value / totalFrameWeight;
+			}
+
+			// Output HTML
+			if (htmlFile != null)
+			{
+				htmlFile.WriteLine("  <br><h3>" + title + "</h3>");
+				htmlFile.WriteLine("  <table border='0' bgcolor='#000000' style='width:1400'>");
+				htmlFile.WriteLine("  <tr>");
+				foreach (Column col in filteredColumns)
+				{
+					htmlFile.WriteLine("<td bgcolor='#ffffff'><b>" + col.name + "</b></td>");
+				}
+				htmlFile.WriteLine("  </tr>");
+				htmlFile.WriteLine("  <tr>");
+				foreach (Column col in filteredColumns)
+				{
+					string bgcolor = "'#ffffff'";
+					if (col.colourThresholdList != null)
+					{
+						bgcolor = col.colourThresholdList.GetColourForValue(col.value);
+					}
+					htmlFile.WriteLine("<td bgcolor=" + bgcolor + ">" + col.value.ToString("0.00") + "</td>");
+				}
+				htmlFile.WriteLine("  </tr>");
+				htmlFile.WriteLine("  </table>");
+			}
+
+			// Output metadata
+			if (metadata != null)
+			{
+				foreach (Column col in filteredColumns)
+				{
+					if ( col.metadataKey != null )
+					{
+						metadata.Add(col.metadataKey, col.value.ToString("0.00"), col.colourThresholdList);
+					}
+				}
+			}
+		}
+		public override void PostInit(ReportTypeInfo reportTypeInfo, CsvStats csvStats)
+		{
+		}
+		string title;
+		string beginEvent;
+		string endEvent;
+		double endOffsetPercentage;
+		List<Column> columns;
+	};
+
+
+	class SummaryMetadataValue
     {
 		public SummaryMetadataValue(double inValue, ColourThresholdList inColorThresholdList, string inToolTip)
 		{
