PerfReportTool 3.91
- Support for map overlay summaries
- Fix collated/email display when there are missing values. Instead of displaying nothing for that entry, compute min/max/avg only for values that exist
- Reduce report size by limiting decimal places for graphs

CsvToSVG 2.20
- Add support for limiting decimal places for CSV graphs

diff --git a/Engine/Source/Programs/CSVTools/CSVToSVG/CsvToSVG.cs b/Engine/Source/Programs/CSVTools/CSVToSVG/CsvToSVG.cs
--- a/Engine/Source/Programs/CSVTools/CSVToSVG/CsvToSVG.cs
+++ b/Engine/Source/Programs/CSVTools/CSVToSVG/CsvToSVG.cs
@@ -24,7 +24,7 @@ namespace CSVTools
 {
     class Version
     {
-        private static string VersionString = "2.18";
+        private static string VersionString = "2.20";
         
         public static string Get() { return VersionString; }
     };
@@ -222,6 +222,7 @@ namespace CSVTools
             "       -percentile \n" +
             "       -percentile90 \n" +
             "       -percentile99 \n" +
+			"       -lineDecimalPlaces <N> (default 3)" +
             "";
 
 
@@ -1651,6 +1652,8 @@ namespace CSVTools
             int n = Math.Min(samples.Count,(int)(range.MaxX+0.5));
             int start = (int)(range.MinX+0.5);
 
+			int numDecimalPlaces = GetIntArg("lineDecimalPlaces", 3);
+
             if (start + 1 < n)
             {
                 List<Vec2> RawPoints = new List<Vec2>();
@@ -1668,14 +1671,22 @@ namespace CSVTools
                     RawPoints = CompressPoints(RawPoints, compression, rect, range);
                 }
 
+				string formatStr = "0";
+				if (numDecimalPlaces>0)
+				{
+					formatStr += ".";
+					for (int i=0; i<numDecimalPlaces;i++)
+					{
+						formatStr += "0";
+					}
+				}
                 string idString = id.Length > 0 ? "id='" + id + "'" : "";
                 SvgWriteLine("<polyline "+idString+" points='");
                 foreach (Vec2 point in RawPoints)
                 {
-                    float x = point.X;//ToSvgX((float)i, rect, range);
-                    float y = point.Y;//ToSvgY(sample, rect, range);
-
-                    SvgWrite(" " + x + "," + y,false);
+                    float x = point.X;
+                    float y = point.Y;
+                    SvgWrite(" " + x.ToString(formatStr) + "," + y.ToString(formatStr), false);
                 }
 
                 string fillString = "none";
diff --git a/Engine/Source/Programs/CSVTools/PerfReportTool/PerfReportTool.cs b/Engine/Source/Programs/CSVTools/PerfReportTool/PerfReportTool.cs
--- a/Engine/Source/Programs/CSVTools/PerfReportTool/PerfReportTool.cs
+++ b/Engine/Source/Programs/CSVTools/PerfReportTool/PerfReportTool.cs
@@ -19,7 +19,7 @@ namespace PerfReportTool
 {
     class Version
     {
-        private static string VersionString = "3.90";
+        private static string VersionString = "3.91";
 
         public static string Get() { return VersionString; }
     };
@@ -1140,6 +1140,8 @@ namespace PerfReportTool
 				snapToPeaks = graphSettings.snapToPeaks.value;
 			}
 
+			int lineDecimalPlaces = graphSettings.lineDecimalPlaces.isSet ? graphSettings.lineDecimalPlaces.value : 1;
+
 			int maxHierarchyDepth = graphSettings.maxHierarchyDepth.value;
             string hideStatPrefix = graphSettings.hideStatPrefix.value;
             string showEvents = graphSettings.showEvents.value;
@@ -1208,8 +1210,9 @@ namespace PerfReportTool
                 " -height "  + (height*scaleby).ToString() +
                 " -budget "  + budget.ToString() +
                 " -maxy "    + maxy.ToString()   +
+				" -lineDecimalPlaces "+ lineDecimalPlaces.ToString() +
 
-                ( ( statMultiplier != 1.0)   ? " -statMultiplier " + statMultiplier.ToString("0.0000000000000000000000") : "") +
+				( ( statMultiplier != 1.0)   ? " -statMultiplier " + statMultiplier.ToString("0.0000000000000000000000") : "") +
                 (hideEventNames              ? " -hideeventNames 1" : "") +
                 ( ( minx > 0 )               ? ( " -minx " + minx.ToString() ) : "" ) +
                 ( ( maxx != Int32.MaxValue ) ? ( " -maxx " + maxx.ToString() ) : "" ) +
@@ -1318,7 +1321,7 @@ namespace PerfReportTool
 
     class ReportTypeInfo
     {
-        public ReportTypeInfo(XElement element, Dictionary<string,XElement> sharedSummaries)
+        public ReportTypeInfo(XElement element, Dictionary<string,XElement> sharedSummaries, string baseXmlDirectory)
         {
             graphs = new List<ReportGraph>();
             summaries = new List<Summary>();
@@ -1344,27 +1347,31 @@ namespace PerfReportTool
 					string summaryType = summaryElement.Attribute("type").Value;
 					if (summaryType == "histogram")
 					{
-						summaries.Add(new HistogramSummary(summaryElement));
+						summaries.Add(new HistogramSummary(summaryElement, baseXmlDirectory));
 					}
 					else if (summaryType == "peak")
 					{
-						summaries.Add(new PeakSummary(summaryElement));
+						summaries.Add(new PeakSummary(summaryElement, baseXmlDirectory));
 					}
 					else if (summaryType == "fpschart")
 					{
-						summaries.Add(new FPSChartSummary(summaryElement));
+						summaries.Add(new FPSChartSummary(summaryElement, baseXmlDirectory));
 					}
 					else if (summaryType == "hitches")
 					{
-						summaries.Add(new HitchSummary(summaryElement));
+						summaries.Add(new HitchSummary(summaryElement, baseXmlDirectory));
 					}
 					else if (summaryType == "event")
 					{
-						summaries.Add(new EventSummary(summaryElement));
+						summaries.Add(new EventSummary(summaryElement, baseXmlDirectory));
 					}
 					else if (summaryType == "boundedstatvalues")
 					{
-						summaries.Add(new BoundedStatValuesSummary(summaryElement));
+						summaries.Add(new BoundedStatValuesSummary(summaryElement, baseXmlDirectory));
+					}
+					else if (summaryType == "mapoverlay")
+					{
+						summaries.Add(new MapOverlaySummary(summaryElement, baseXmlDirectory));
 					}
 				}
 				else if (child.Name == "metadataToShow")
@@ -1434,6 +1441,7 @@ namespace PerfReportTool
             statMultiplier = new OptionalDouble(element, "statMultiplier");
 			legendAverageThreshold = new OptionalDouble(element, "legendAverageThreshold");
 			snapToPeaks = new OptionalBool(element, "snapToPeaks");
+			lineDecimalPlaces = new OptionalInt(element, "lineDecimalPlaces");
 		}
 		public void InheritFrom(GraphSettings baseSettings)
         {
@@ -1464,6 +1472,7 @@ namespace PerfReportTool
             ignoreStats.InheritFrom(baseSettings.ignoreStats);
 			legendAverageThreshold.InheritFrom(baseSettings.legendAverageThreshold);
 			snapToPeaks.InheritFrom(baseSettings.snapToPeaks);
+			lineDecimalPlaces.InheritFrom(baseSettings.lineDecimalPlaces);
 
 		}
         public OptionalBool smooth;
@@ -1494,6 +1503,7 @@ namespace PerfReportTool
 
 		public OptionalBool requiresDetailedStats;
 		public OptionalBool snapToPeaks;
+		public OptionalInt lineDecimalPlaces;
 
 	};
 
@@ -1605,6 +1615,8 @@ namespace PerfReportTool
 			}
 			Console.Out.WriteLine("BaseDir: " + baseDirectory);
 
+			baseXmlDirectory = baseDirectory;
+
 			// Read the report type XML
 			reportTypeXmlFilename = Path.Combine(baseDirectory, "reportTypes.xml" );
 			if (reportXMLFilenameIn.Length > 0)
@@ -1780,7 +1792,7 @@ namespace PerfReportTool
 				{
 					if (IsReportTypeXMLCompatibleWithStats(element, csvFile.dummyCsvStats))
 					{
-						reportTypeInfo = new ReportTypeInfo(element, sharedSummaries);
+						reportTypeInfo = new ReportTypeInfo(element, sharedSummaries, baseXmlDirectory);
 						break;
 					}
 				}
@@ -1815,7 +1827,7 @@ namespace PerfReportTool
                         throw new Exception("Report type " + reportType + " is not compatible with CSV " + csvFile.filename);
                     }
                 }
-                reportTypeInfo = new ReportTypeInfo(foundReportTypeElement,sharedSummaries);
+                reportTypeInfo = new ReportTypeInfo(foundReportTypeElement, sharedSummaries, baseXmlDirectory);
             }
 
             // Load the graphs
@@ -1943,6 +1955,7 @@ namespace PerfReportTool
 		Dictionary<string, GraphSettings> graphs;
 		Dictionary<string, string> statDisplayNameMapping;
 		DerivedMetadataMappings derivedMetadataMappings;
+		string baseXmlDirectory;
 
 		List<CsvEventStripInfo> csvEventsToStrip;
         string reportTypeXmlFilename;
diff --git a/Engine/Source/Programs/CSVTools/PerfReportTool/perfsummaries.cs b/Engine/Source/Programs/CSVTools/PerfReportTool/perfsummaries.cs
--- a/Engine/Source/Programs/CSVTools/PerfReportTool/perfsummaries.cs
+++ b/Engine/Source/Programs/CSVTools/PerfReportTool/perfsummaries.cs
@@ -386,7 +386,7 @@ namespace PerfSummaries
 
     class FPSChartSummary : Summary
     {
-        public FPSChartSummary(XElement element)
+        public FPSChartSummary(XElement element, string baseXmlDirectory)
         {
             ReadStatsFromXML(element);
             fps = Convert.ToInt32(element.Attribute("fps").Value);
@@ -719,7 +719,7 @@ namespace PerfSummaries
 
     class EventSummary : Summary
     {
-        public EventSummary(XElement element)
+        public EventSummary(XElement element, string baseXmlDirectory)
         {
             title = element.GetSafeAttibute("title","Events");
             metadataKey = element.Attribute("metadataKey").Value;
@@ -797,7 +797,7 @@ namespace PerfSummaries
 
     class HitchSummary : Summary
     {
-        public HitchSummary(XElement element)
+        public HitchSummary(XElement element, string baseXmlDirectory)
         {
             ReadStatsFromXML(element);
 
@@ -892,7 +892,7 @@ namespace PerfSummaries
 
     class HistogramSummary : Summary
     {
-        public HistogramSummary(XElement element)
+        public HistogramSummary(XElement element, string baseXmlDirectory)
         {
             ReadStatsFromXML(element);
 
@@ -1047,7 +1047,7 @@ namespace PerfSummaries
 
     class PeakSummary : Summary
     {
-        public PeakSummary(XElement element)
+        public PeakSummary(XElement element, string baseXmlDirectory)
         {
 			hidePrefixes = new List<string>();
 			sectionPrefixes = new List<string>();
@@ -1308,7 +1308,7 @@ namespace PerfSummaries
 			public double threshold;
 			public ColourThresholdList colourThresholdList;
 		};
-		public BoundedStatValuesSummary(XElement element)
+		public BoundedStatValuesSummary(XElement element, string baseXmlDirectory)
 		{
 			ReadStatsFromXML(element);
 			if (stats.Count != 0)
@@ -1382,11 +1382,13 @@ namespace PerfSummaries
 			}
 			if ( startFrame == -1 )
 			{
-				throw new Exception("BoundedStatValuesSummary: Begin event " + beginEvent+" was not found");
+				Console.WriteLine("BoundedStatValuesSummary: Begin event " + beginEvent+" was not found");
+				return;
 			}
-			if (endFrame== -1)
+			if (endFrame==int.MaxValue)
 			{
-				throw new Exception("BoundedStatValuesSummary: End event " + endEvent + " was not found");
+				Console.WriteLine("BoundedStatValuesSummary: End event " + endEvent + " was not found");
+				return;
 			}
 			if ( startFrame >= endFrame )
 			{
@@ -1561,6 +1563,147 @@ namespace PerfSummaries
 		List<Column> columns;
 	};
 
+	class MapOverlaySummary : Summary
+	{
+		class MapOverlayEvent
+		{
+			public string name;
+			public string metadataPrefix;
+		};
+		public MapOverlaySummary(XElement element, string baseXmlDirectory)
+		{
+			ReadStatsFromXML(element);
+			if (stats.Count != 0)
+			{
+				throw new Exception("<stats> element is not supported");
+			}
+
+			xStatName = element.GetSafeAttibute<string>("xStat");
+			yStatName = element.GetSafeAttibute<string>("yStat");
+			zStatName = element.GetSafeAttibute<string>("zStat");
+			stats.Add(xStatName);
+			stats.Add(yStatName);
+			stats.Add(zStatName);
+			sourceImagePath = element.GetSafeAttibute<string>("sourceImage");
+			if ( !System.IO.Path.IsPathRooted(sourceImagePath))
+			{
+				sourceImagePath = System.IO.Path.GetFullPath(System.IO.Path.Combine(baseXmlDirectory,sourceImagePath));
+			}
+
+			offsetX = element.GetSafeAttibute<float>("offsetX",0.0f);
+			offsetY = element.GetSafeAttibute<float>("offsetY",0.0f);
+			scale = element.GetSafeAttibute<float>("scale",1.0f);
+			title = element.GetSafeAttibute("title", "Events");
+			destImageFilename = element.Attribute("destImage").Value;
+			imageWidth = element.GetSafeAttibute<float>("width", 250.0f);
+			imageHeight = element.GetSafeAttibute<float>("height", 250.0f);
+			framesPerLineSegment = element.GetSafeAttibute<int>("framesPerLineSegment", 5);
+
+			foreach (XElement eventEl in element.Elements("event"))
+			{
+				MapOverlayEvent ev = new MapOverlayEvent();
+				ev.name = eventEl.Attribute("name").Value;
+				ev.metadataPrefix = eventEl.GetSafeAttibute<string>("metadataPrefix");
+				events.Add(ev);
+			}
+		}
+
+		int toSvgX(float worldX, float worldY)
+		{
+			float svgX = (worldY * scale + offsetX) * 0.5f + 0.5f;
+			svgX *= imageWidth;
+			return (int)(svgX + 0.5f);
+		}
+
+		int toSvgY(float worldX, float worldY)
+		{ 
+			float svgY = 1.0f - (worldX * scale + offsetY) * 0.5f - 0.5f;
+			svgY *= imageHeight;
+			return (int)(svgY + 0.5f);
+		}
+
+		public override void WriteSummaryData(System.IO.StreamWriter htmlFile, CsvStats csvStats, bool bIncludeSummaryCsv, SummaryMetadata metadata, string htmlFileName)
+		{
+			// Output HTML
+			if (htmlFile != null)
+			{
+				//string outputDirectory= System.IO.Path.GetDirectoryName(System.IO.Path.GetFullPath(htmlFileName));
+
+				//string outputMapFilename = System.IO.Path.Combine(outputDirectory, destImageFilename);
+
+				if ( !System.IO.File.Exists(destImageFilename))
+				{
+					System.IO.File.Copy(sourceImagePath, destImageFilename);
+				}
+
+				StatSamples xStat = csvStats.GetStat(xStatName);
+				StatSamples yStat = csvStats.GetStat(yStatName);
+
+				if ( xStat == null || yStat == null )
+				{
+					throw new Exception("MapOverlaySummary: CSV doesn't contain stat " + xStatName + " and/or " + yStatName);
+				}
+
+				// Check if the file exists in the output directory
+				htmlFile.WriteLine("  <br><h3>" + title + "</h3>");
+				htmlFile.WriteLine("<svg version='1.1' xmlns='http://www.w3.org/2000/svg' xmlns:xlink='http://www.w3.org/1999/xlink' width='" + imageWidth + "' height='" + imageHeight + "'>");
+				htmlFile.WriteLine("<image href='" + destImageFilename + "' width='" + imageWidth + "' height='" + imageHeight + "' />");
+
+				//List<string> polyLineStrings = new List<string>();
+				htmlFile.Write("<polyline points='");
+				for ( int i=0; i<xStat.samples.Count;i+=framesPerLineSegment)
+				{
+					float x = xStat.samples[i];
+					float y = yStat.samples[i];
+					htmlFile.Write(toSvgX(x,y)+","+toSvgY(x,y)+" ");
+				}
+				string colourString = "#ff0000";
+				htmlFile.WriteLine("' style='fill:none;stroke-width:1.3; clip-path: url(#graphArea)' stroke='" + colourString + "'/>");
+				float circleRadius = 3;
+				string eventColourString = "#ffffff";
+				foreach (MapOverlayEvent mapEvent in events)
+				{
+					foreach (CsvEvent ev in csvStats.Events)
+					{
+						if (CsvStats.DoesSearchStringMatch(ev.Name, mapEvent.name))
+						{
+							float x = xStat.samples[ev.Frame];
+							float y = yStat.samples[ev.Frame];
+							int svgX = toSvgX(x, y);
+							int svgY = toSvgY(x, y);
+							htmlFile.Write("<circle cx='" + svgX + "' cy='" + svgY + "' r='" + circleRadius + "' fill='" + eventColourString + "' fill-opacity='1.0'/>");
+							htmlFile.WriteLine("<text x='"+(svgX+5)+"' y='"+svgY+"' text-anchor='left' style='font-family: Verdana;fill: #ffffff; font-size: " + 9 + "px;'>" + ev.Name + "</text>");
+						}
+					}
+				}
+
+				htmlFile.WriteLine("<text x='50%' y='" + (imageHeight * 0.05) + "' text-anchor='middle' style='font-family: Verdana;fill: #FFFFFF; stroke: #C0C0C0;  font-size: " + 20 + "px;'>" + title + "</text>");
+				htmlFile.WriteLine("</svg>");
+			}
+
+			// Output metadata
+			if (metadata != null)
+			{
+			}
+		}
+		public override void PostInit(ReportTypeInfo reportTypeInfo, CsvStats csvStats)
+		{
+		}
+		string title;
+		string xStatName;
+		string yStatName;
+		string zStatName;
+		string sourceImagePath;
+		float offsetX;
+		float offsetY;
+		float scale;
+		string destImageFilename;
+		float imageWidth;
+		float imageHeight;
+		int framesPerLineSegment;
+
+		List<MapOverlayEvent> events = new List<MapOverlayEvent>();
+	};
 
 	class SummaryMetadataValue
     {
@@ -1898,6 +2041,7 @@ namespace PerfSummaries
 			List<float> RowMaxValues = new List<float>();
 			List<float> RowTotals = new List<float>();
 			List<float> RowMinValues = new List<float>();
+			List<int> RowValueCounts = new List<int>();
 
 			// Set the initial sort key
 			string CurrentRowSortKey = "";
@@ -1916,11 +2060,13 @@ namespace PerfSummaries
 					RowMaxValues.Clear();
 					RowMinValues.Clear();
 					RowTotals.Clear();
+					RowValueCounts.Clear();
 					for (int j = 0; j < columns.Count; j++)
 					{
 						RowMaxValues.Add(-float.MaxValue);
 						RowMinValues.Add(float.MaxValue);
 						RowTotals.Add(0.0f);
+						RowValueCounts.Add(0);
 					}
 					mergedRowsCount = 0;
 					reset = false;
@@ -1935,13 +2081,14 @@ namespace PerfSummaries
 						float value = column.GetValue(i);
 						if (value == float.MaxValue)
 						{
-							RowTotals[j] = float.MaxValue;
+							//RowTotals[j] = float.MaxValue;
 						}
 						else
 						{
 							RowMaxValues[j] = Math.Max(RowMaxValues[j], value);
 							RowMinValues[j] = Math.Min(RowMinValues[j], value);
 							RowTotals[j] += value;
+							RowValueCounts[j]++;
 						}
 					}
 				}
@@ -1973,7 +2120,7 @@ namespace PerfSummaries
 						int destColumnBaseIndex = srcToDestBaseColumnIndex[j];
 						if (destColumnBaseIndex != -1 && RowTotals[j] != float.MaxValue)
 						{
-							newColumns[destColumnBaseIndex].SetValue(destRowIndex, RowTotals[j] / (float)mergedRowsCount);
+							newColumns[destColumnBaseIndex].SetValue(destRowIndex, RowTotals[j] / (float)RowValueCounts[j]);
 							newColumns[destColumnBaseIndex+1].SetValue(destRowIndex, RowMinValues[j]);
 							newColumns[destColumnBaseIndex+2].SetValue(destRowIndex, RowMaxValues[j]);
 
