PerfReportTool 3.92
Map overlay improvements:
- format improvements
- line color event support
- start event
- line strip resets based on distance
- support for unstripped CSV for certain summaries (e.g mapoverlay and file IO), configurable via XML
- stripped stats with total of zero (workaround for texture streaming stats getting set with zeros in shipping)

diff --git a/Engine/Source/Programs/CSVTools/CSVSplit/CSVSplit.cs b/Engine/Source/Programs/CSVTools/CSVSplit/CSVSplit.cs
--- a/Engine/Source/Programs/CSVTools/CSVSplit/CSVSplit.cs
+++ b/Engine/Source/Programs/CSVTools/CSVSplit/CSVSplit.cs
@@ -98,7 +98,7 @@ namespace CSVTools
                     destStats.metaData = csvStats.metaData;
                     foreach (StatSamples stat in sourceStats)
                     {
-                        StatSamples newStat = new StatSamples(stat);
+                        StatSamples newStat = new StatSamples(stat,false);
                         destStats.AddStat(newStat);
                     }
                     switchValueToCsvIndex.Add(switchValue, csvStatsOutList.Count);
diff --git a/Engine/Source/Programs/CSVTools/CSVToSVG/CsvToSVG.cs b/Engine/Source/Programs/CSVTools/CSVToSVG/CsvToSVG.cs
--- a/Engine/Source/Programs/CSVTools/CSVToSVG/CsvToSVG.cs
+++ b/Engine/Source/Programs/CSVTools/CSVToSVG/CsvToSVG.cs
@@ -24,7 +24,7 @@ namespace CSVTools
 {
     class Version
     {
-        private static string VersionString = "2.20";
+        private static string VersionString = "2.21";
         
         public static string Get() { return VersionString; }
     };
@@ -1123,7 +1123,7 @@ namespace CSVTools
             CsvStats smoothStats = new CsvStats();
             foreach( StatSamples srcStatSamples in stats.Stats.Values )
             {
-                StatSamples NewStatSamples = new StatSamples(srcStatSamples);
+                StatSamples NewStatSamples = new StatSamples(srcStatSamples,false);
                 smoothStats.AddStat(NewStatSamples);
             }
             StatSamples[] SmoothSamplesArray = smoothStats.Stats.Values.ToArray();
@@ -1255,7 +1255,7 @@ namespace CSVTools
                 StatSamplesSorted.Add(OtherStat);
                 foreach (float value in totalStat.samples)
                     OtherStat.samples.Add(value);
-                OtherUnstacked = new StatSamples(OtherStat);
+                OtherUnstacked = new StatSamples(OtherStat,false);
                 rangeEnd = Math.Min(OtherStat.samples.Count, (int)range.MaxX);
             }
 
@@ -1263,7 +1263,7 @@ namespace CSVTools
             List<StatSamples> StatSamplesList = new List<StatSamples>();
             foreach (StatSamples srcStatSamples in StatSamplesSorted)
             {
-                StatSamples destStatSamples = new StatSamples(srcStatSamples);
+                StatSamples destStatSamples = new StatSamples(srcStatSamples,false);
                 for (int j = 0; j < srcStatSamples.samples.Count; j++)
                 {
                     float value=srcStatSamples.samples[j]; 
@@ -2057,7 +2057,7 @@ namespace CSVTools
                     statInfo.jsTextElementId = "csvStatText__" + i + GetJSStatName(stat.Name) + "<UNIQUE>";
                     statInfo.jsGroupElementId = "csvStatGroup__" + i + GetJSStatName(stat.Name) + "<UNIQUE>";
                     statInfo.originalStatSamples = stat;
-                    statInfo.statSamples = new StatSamples(stat);
+                    statInfo.statSamples = new StatSamples(stat,false);
                     interactiveStats.Add(statInfo);
                 }
             }
diff --git a/Engine/Source/Programs/CSVTools/CsvStats/CsvStats.cs b/Engine/Source/Programs/CSVTools/CsvStats/CsvStats.cs
--- a/Engine/Source/Programs/CSVTools/CsvStats/CsvStats.cs
+++ b/Engine/Source/Programs/CSVTools/CsvStats/CsvStats.cs
@@ -15,7 +15,12 @@ namespace CSVStats
 {
     public class CsvEvent
     {
-        public CsvEvent(string NameIn, int FrameIn)
+		public CsvEvent(CsvEvent source)
+		{
+			Name = source.Name;
+			Frame = source.Frame;
+		}
+		public CsvEvent(string NameIn, int FrameIn)
         {
             Name = NameIn;
             Frame = FrameIn;
@@ -24,6 +29,7 @@ namespace CSVStats
         {
 
         }
+
         public string Name;
         public int Frame;
     };
@@ -39,7 +45,17 @@ namespace CSVStats
             return rv;
         }
 
-        public CsvMetadata()
+		public CsvMetadata(CsvMetadata source)
+		{
+			Values = new Dictionary<string, string>();
+			List<KeyValuePair<string, string>> pairList = source.Values.ToList();
+			foreach (KeyValuePair<string, string> pair in pairList)
+			{
+				Values.Add(pair.Key, pair.Value);
+			}
+		}
+
+		public CsvMetadata()
 		{
 			Values = new Dictionary<string, string>();
 		}
@@ -193,11 +209,15 @@ namespace CSVStats
 
     public class StatSamples : IComparable
     {
-        public StatSamples(StatSamples source)
+        public StatSamples(StatSamples source, bool copySamples=true)
         {
             Name = source.Name;
             samples = new List<float>();
-            average = source.average;
+			if (copySamples)
+			{
+				samples.AddRange(source.samples);
+			}
+			average = source.average;
             total = source.total;
             colour = source.colour;
             LegendName = source.LegendName;
@@ -352,7 +372,24 @@ namespace CSVStats
             Stats = new Dictionary<string, StatSamples>();
             Events = new List<CsvEvent>();
         }
-        public StatSamples GetStat(string name)
+		public CsvStats(CsvStats source)
+		{
+			Stats = new Dictionary<string, StatSamples>();
+			Events = new List<CsvEvent>();
+			foreach (StatSamples sourceStat in source.Stats.Values)
+			{
+				AddStat(new StatSamples(sourceStat));
+			}
+			foreach (CsvEvent ev in source.Events)
+			{
+				Events.Add(new CsvEvent(ev));
+			}
+			if (source.metaData != null)
+			{
+				metaData = new CsvMetadata(source.metaData);
+			}
+		}
+		public StatSamples GetStat(string name)
         {
             name = name.ToLower();
             if (Stats.ContainsKey(name))
@@ -866,22 +903,25 @@ namespace CSVStats
             }
         }
 
-        public int StripByEvents(string startString, string endString, bool invert=false)
+        public CsvStats StripByEvents(string startString, string endString, bool invert, out int numFramesStripped)
         {
+			CsvStats newCsvStats = new CsvStats();
             List<int> startIndices = null;
             List<int> endIndices = null;
             GetEventFrameIndexDelimiters(startString, endString, out startIndices, out endIndices);
+			numFramesStripped = 0;
 			if (startIndices.Count == 0 )
             {
-                return 0;
+                return this;
             }
 
-            int framesStripped = -1;
-            int frameCount = 0;
+			numFramesStripped = -1;
+			int frameCount = 0;
             // Strip out samples and recompute averages
             foreach (StatSamples stat in Stats.Values)
             {
-                List<float> newSamples = new List<float>(SampleCount);
+				StatSamples newStat = new StatSamples(stat, false);
+				newStat.samples = new List<float>(SampleCount);
                 int stripEventIndex = 0;
                 for (int i = 0; i < stat.samples.Count; i++)
                 {
@@ -889,7 +929,7 @@ namespace CSVStats
                     int endIndex = (stripEventIndex < endIndices.Count) ? endIndices[stripEventIndex] : stat.samples.Count;
                     if (i < startIndex)
                     {
-                        newSamples.Add(stat.samples[i]);
+						newStat.samples.Add(stat.samples[i]);
                     }
                     else
                     {
@@ -899,31 +939,29 @@ namespace CSVStats
                         }
                     }
                 }
-                if (framesStripped == -1)
+                if (numFramesStripped == -1)
                 {
-                    framesStripped = stat.samples.Count - newSamples.Count;
+					numFramesStripped = stat.samples.Count - newStat.samples.Count;
                     frameCount = stat.samples.Count;
                 }
-                stat.samples = newSamples;
-                stat.ComputeAverageAndTotal();
+
+                newStat.ComputeAverageAndTotal();
+				newCsvStats.AddStat(newStat);
             }
 
             // Strip out the events
             int FrameOffset = 0;
             {
-                List<CsvEvent> newEvents = new List<CsvEvent>();
                 int stripEventIndex = 0;
                 for (int i = 0; i < Events.Count; i++)
                 {
                     CsvEvent csvEvent = Events[i];
                     int startIndex = (stripEventIndex < startIndices.Count) ? startIndices[stripEventIndex] : frameCount;
                     int endIndex = (stripEventIndex < endIndices.Count) ? endIndices[stripEventIndex] : frameCount;
-                    CsvEvent newEvent = new CsvEvent();
-                    newEvent.Frame = csvEvent.Frame + FrameOffset;
-                    newEvent.Name = csvEvent.Name;
+                    CsvEvent newEvent = new CsvEvent(csvEvent.Name, csvEvent.Frame + FrameOffset);
                     if (csvEvent.Frame < startIndex)
                     {
-                        newEvents.Add(newEvent);
+						newCsvStats.Events.Add(newEvent);
                     }
                     else
                     {
@@ -936,18 +974,17 @@ namespace CSVStats
 								FrameOffset -= endIndex - startIndex;
 							}
                         }
-
                         if (csvEvent.Frame == endIndex+1)
                         {
-                            newEvents.Add(newEvent);
+							newCsvStats.Events.Add(newEvent);
                             stripEventIndex++;
                         }
                     }
                 }
-                Events = newEvents;
             }
-            return framesStripped;
-        }
+			newCsvStats.metaData = metaData;
+			return newCsvStats;
+		}
 
         public int SampleCount
         {
diff --git a/Engine/Source/Programs/CSVTools/PerfReportTool/PerfReportTool.cs b/Engine/Source/Programs/CSVTools/PerfReportTool/PerfReportTool.cs
--- a/Engine/Source/Programs/CSVTools/PerfReportTool/PerfReportTool.cs
+++ b/Engine/Source/Programs/CSVTools/PerfReportTool/PerfReportTool.cs
@@ -19,7 +19,7 @@ namespace PerfReportTool
 {
     class Version
     {
-        private static string VersionString = "3.91";
+        private static string VersionString = "3.92";
 
         public static string Get() { return VersionString; }
     };
@@ -725,8 +725,16 @@ namespace PerfReportTool
 
             // Read the csv stats while we wait for the graphs to complete
             int numFramesStripped = 0;
-            CsvStats csvStats = ReadCsvStats(csvFile, minX, maxX, out numFramesStripped);
-            perfLog.LogTiming("    ReadCsvStats");
+            CsvStats csvStats = ReadCsvStats(csvFile, minX, maxX);
+			CsvStats unstrippedCsvStats = csvStats;
+
+			if (!GetBoolArg("noStripEvents"))
+			{
+				CsvStats strippedCsvStats = StripCsvStatsByEvents(unstrippedCsvStats, out numFramesStripped);
+				csvStats = strippedCsvStats;
+			}
+
+			perfLog.LogTiming("    ReadCsvStats");
 
             if ( writeDetailedReport )
             { 
@@ -755,6 +763,10 @@ namespace PerfReportTool
 					if (!statsToSummarise.Contains(ListOfKeys[i]))
 					{
 						csvStats.Stats.Remove(ListOfKeys[i]);
+						if (csvStats != unstrippedCsvStats)
+						{
+							unstrippedCsvStats.Stats.Remove(ListOfKeys[i]);
+						}
 					}
 				}
 			}
@@ -812,7 +824,7 @@ namespace PerfReportTool
             }
 
             // Write the report
-            WriteReport(htmlFilename, title, svgFilenames, reportTypeInfo, csvStats, numFramesStripped, minX, maxX, bBulkMode, summaryMetadata);
+            WriteReport(htmlFilename, title, svgFilenames, reportTypeInfo, csvStats, unstrippedCsvStats, numFramesStripped, minX, maxX, bBulkMode, summaryMetadata);
             perfLog.LogTiming("    WriteReport");
 
             // Delete the temp files
@@ -824,12 +836,12 @@ namespace PerfReportTool
 				}
             }
         }
-        CsvStats ReadCsvStats(CachedCsvFile csvFile, int minX, int maxX, out int numFramesStripped)
-        {
-            CsvStats csvStats = CsvStats.ReadCSVFromLines(csvFile.lines, null);
+		CsvStats ReadCsvStats(CachedCsvFile csvFile, int minX, int maxX)
+		{
+			CsvStats csvStats = CsvStats.ReadCSVFromLines(csvFile.lines, null);
 			reportXML.ApplyDerivedMetadata(csvStats.metaData);
 
-			if ( csvStats.metaData == null )
+			if (csvStats.metaData == null)
 			{
 				csvStats.metaData = new CsvMetadata();
 			}
@@ -849,8 +861,8 @@ namespace PerfReportTool
 				}
 			}
 			string endEventStr = GetArg("endEvent").ToLower();
-			if ( endEventStr != "")
-			{ 
+			if (endEventStr != "")
+			{
 				for (int i = csvStats.Events.Count - 1; i >= 0; i--)
 				{
 					CsvEvent ev = csvStats.Events[i];
@@ -862,35 +874,53 @@ namespace PerfReportTool
 				}
 			}
 
+			// Strip out all stats with a zero total
+			List<StatSamples> allStats = new List<StatSamples>();
+			foreach (StatSamples stat in csvStats.Stats.Values)
+			{
+				allStats.Add(stat);
+			}
+			csvStats.Stats.Clear();
+			foreach (StatSamples stat in allStats)
+			{
+				if (stat.total != 0.0f)
+				{
+					csvStats.AddStat(stat);
+				}
+			}
+
 			// Crop the stats to the range
 			csvStats.CropStats(minX, maxX);
+			return csvStats;
+		}
 
-            numFramesStripped = 0;
+		CsvStats StripCsvStatsByEvents(CsvStats csvStats, out int numFramesStripped)
+		{
+			numFramesStripped = 0;
             List<CsvEventStripInfo> eventsToStrip = reportXML.GetCsvEventsToStrip();
-            if (!GetBoolArg("noStripEvents"))
+			CsvStats strippedStats = csvStats;
+            if (eventsToStrip != null)
             {
-                if (eventsToStrip != null)
+				// This can have issues if we strip events and then subsequently strip overlapping events. We'd get better results if we did it in a single pass
+				foreach (CsvEventStripInfo eventStripInfo in eventsToStrip)
                 {
-                    // This can have issues if we strip events and then subsequently strip overlapping events. We'd get better results here if we did this in a single pass
-                    foreach (CsvEventStripInfo eventStripInfo in eventsToStrip)
-                    {
-                        numFramesStripped += csvStats.StripByEvents(eventStripInfo.beginName, eventStripInfo.endName);
-                    }
-                }
+					int numFramesStrippedThisStage = 0;
+					strippedStats = strippedStats.StripByEvents(eventStripInfo.beginName, eventStripInfo.endName, false, out numFramesStrippedThisStage);
+					numFramesStripped += numFramesStrippedThisStage;
+				}
             }
 
             if (numFramesStripped > 0 )
             {
                 Console.WriteLine("CSV frames excluded : " + numFramesStripped);
             }
-
-            return csvStats;
+            return strippedStats;
         }
 
       
 
 
-        void WriteReport(string htmlFilename, string title, List<string> svgFilenames, ReportTypeInfo reportTypeInfo, CsvStats csvStats, int numFramesStripped, int minX, int maxX, bool bBulkMode, SummaryMetadata summaryMetadata)
+        void WriteReport(string htmlFilename, string title, List<string> svgFilenames, ReportTypeInfo reportTypeInfo, CsvStats csvStats, CsvStats unstrippedCsvStats, int numFramesStripped, int minX, int maxX, bool bBulkMode, SummaryMetadata summaryMetadata)
         {
  
             ReportGraph[] graphs = reportTypeInfo.graphs.ToArray();
@@ -964,7 +994,7 @@ namespace PerfReportTool
             PeakSummary peakSummary = null;
             foreach (Summary summary in reportTypeInfo.summaries)
             {
-                summary.WriteSummaryData(htmlFile, csvStats, bIncludeSummaryCsv, summaryMetadata, htmlFilename);
+                summary.WriteSummaryData(htmlFile, summary.useUnstrippedCsvStats ? unstrippedCsvStats : csvStats, bIncludeSummaryCsv, summaryMetadata, htmlFilename);
                 if ( summary.GetType() == typeof(PeakSummary) )
                 {
                     peakSummary = (PeakSummary)summary;
diff --git a/Engine/Source/Programs/CSVTools/PerfReportTool/perfsummaries.cs b/Engine/Source/Programs/CSVTools/PerfReportTool/perfsummaries.cs
--- a/Engine/Source/Programs/CSVTools/PerfReportTool/perfsummaries.cs
+++ b/Engine/Source/Programs/CSVTools/PerfReportTool/perfsummaries.cs
@@ -238,6 +238,7 @@ namespace PerfSummaries
 
         public void ReadStatsFromXML(XElement element)
         {
+			useUnstrippedCsvStats = element.GetSafeAttibute<bool>("useUnstrippedCsvStats", false);
 			XElement statsElement = element.Element("stats");
 			if (statsElement != null)
 			{
@@ -382,6 +383,7 @@ namespace PerfSummaries
         public List<CaptureRange> captures;
         public List<string> stats;
         public Dictionary<string, ColourThresholdList> StatThresholds;
+		public bool useUnstrippedCsvStats;
     };
 
     class FPSChartSummary : Summary
@@ -1567,8 +1569,14 @@ namespace PerfSummaries
 	{
 		class MapOverlayEvent
 		{
+			public MapOverlayEvent(string inName)
+			{
+				name = inName;
+			}
 			public string name;
 			public string metadataPrefix;
+			public string shortName;
+			public string lineColor;
 		};
 		public MapOverlaySummary(XElement element, string baseXmlDirectory)
 		{
@@ -1598,12 +1606,22 @@ namespace PerfSummaries
 			imageWidth = element.GetSafeAttibute<float>("width", 250.0f);
 			imageHeight = element.GetSafeAttibute<float>("height", 250.0f);
 			framesPerLineSegment = element.GetSafeAttibute<int>("framesPerLineSegment", 5);
+			lineSplitDistanceThreshold = element.GetSafeAttibute<float>("lineSplitDistanceThreshold", float.MaxValue);
 
 			foreach (XElement eventEl in element.Elements("event"))
 			{
-				MapOverlayEvent ev = new MapOverlayEvent();
-				ev.name = eventEl.Attribute("name").Value;
+				MapOverlayEvent ev = new MapOverlayEvent(eventEl.Attribute("name").Value);
+				ev.shortName = eventEl.GetSafeAttibute<string>("shortName");
 				ev.metadataPrefix = eventEl.GetSafeAttibute<string>("metadataPrefix");
+				ev.lineColor = eventEl.GetSafeAttibute<string>("lineColor");
+				if (eventEl.GetSafeAttibute<bool>("isStartEvent",false))
+				{
+					if (startEvent != null)
+					{
+						throw new Exception("Can't have multiple start events!");
+					}
+					startEvent = ev;
+				}
 				events.Add(ev);
 			}
 		}
@@ -1627,13 +1645,12 @@ namespace PerfSummaries
 			// Output HTML
 			if (htmlFile != null)
 			{
-				//string outputDirectory= System.IO.Path.GetDirectoryName(System.IO.Path.GetFullPath(htmlFileName));
-
-				//string outputMapFilename = System.IO.Path.Combine(outputDirectory, destImageFilename);
+				string outputDirectory= System.IO.Path.GetDirectoryName(System.IO.Path.GetFullPath(htmlFileName));
+				string outputMapFilename = System.IO.Path.Combine(outputDirectory, destImageFilename);
 
-				if ( !System.IO.File.Exists(destImageFilename))
+				if ( !System.IO.File.Exists(outputMapFilename))
 				{
-					System.IO.File.Copy(sourceImagePath, destImageFilename);
+					System.IO.File.Copy(sourceImagePath, outputMapFilename);
 				}
 
 				StatSamples xStat = csvStats.GetStat(xStatName);
@@ -1644,21 +1661,88 @@ namespace PerfSummaries
 					throw new Exception("MapOverlaySummary: CSV doesn't contain stat " + xStatName + " and/or " + yStatName);
 				}
 
+				// If a startevent is specified, update the start frame
+				int startFrame = 0;
+				if (startEvent != null)
+				{
+					foreach (CsvEvent ev in csvStats.Events)
+					{
+						if (CsvStats.DoesSearchStringMatch(ev.Name, startEvent.name))
+						{
+							startFrame = ev.Frame;
+							break;
+						}
+					}
+				}
+
 				// Check if the file exists in the output directory
 				htmlFile.WriteLine("  <br><h3>" + title + "</h3>");
 				htmlFile.WriteLine("<svg version='1.1' xmlns='http://www.w3.org/2000/svg' xmlns:xlink='http://www.w3.org/1999/xlink' width='" + imageWidth + "' height='" + imageHeight + "'>");
-				htmlFile.WriteLine("<image href='" + destImageFilename + "' width='" + imageWidth + "' height='" + imageHeight + "' />");
+				htmlFile.WriteLine("<image href='" + outputMapFilename + "' width='" + imageWidth + "' height='" + imageHeight + "' />");
+
+				// Make a mapping from frame to map indices
+				List<KeyValuePair<int, MapOverlayEvent>> frameEvents = new List<KeyValuePair<int, MapOverlayEvent>>();
+				foreach (MapOverlayEvent mapEvent in events)
+				{
+					foreach (CsvEvent ev in csvStats.Events)
+					{
+						if (CsvStats.DoesSearchStringMatch(ev.Name, mapEvent.name))
+						{
+							frameEvents.Add(new KeyValuePair<int, MapOverlayEvent>(ev.Frame, mapEvent));
+						}
+					}
+				}
+				frameEvents.Sort((pair0, pair1) => pair0.Key.CompareTo(pair1.Key));
+				int eventIndex = 0;
 
 				//List<string> polyLineStrings = new List<string>();
-				htmlFile.Write("<polyline points='");
-				for ( int i=0; i<xStat.samples.Count;i+=framesPerLineSegment)
+				string colourString = "#ff0000";
+				string lineStartTemplate = "<polyline style='fill:none;stroke-width:1.3;stroke:{LINECOLOUR}' points='";
+				htmlFile.Write(lineStartTemplate.Replace("{LINECOLOUR}", colourString));
+				float adjustedLineSplitDistanceThreshold = lineSplitDistanceThreshold * framesPerLineSegment;
+				float oldx=0;
+				float oldy = 0;
+				int lastFrameIndex = 0;
+				for ( int i=startFrame; i<xStat.samples.Count;i+=framesPerLineSegment)
 				{
 					float x = xStat.samples[i];
 					float y = yStat.samples[i];
-					htmlFile.Write(toSvgX(x,y)+","+toSvgY(x,y)+" ");
+					string lineCoordsStr = toSvgX(x, y) + "," + toSvgY(x, y) + " ";
+
+					// Figure out which event we're up to 
+					bool restartLineStrip = false;
+					while (eventIndex < frameEvents.Count && lastFrameIndex < frameEvents[eventIndex].Key && i>= frameEvents[eventIndex].Key)
+					{
+						MapOverlayEvent mapEvent = frameEvents[eventIndex].Value;
+						if (mapEvent.lineColor != null)
+						{
+							colourString = mapEvent.lineColor;
+							restartLineStrip = true;
+						}
+						eventIndex++;
+					}
+
+					float maxManhattanDist = Math.Max(Math.Abs(x - oldx), Math.Abs(y - oldy));
+					if (maxManhattanDist > adjustedLineSplitDistanceThreshold)
+					{
+						restartLineStrip = true;
+					}
+					else
+					{
+						htmlFile.Write(lineCoordsStr);
+					}
+
+					if ( restartLineStrip )
+					{ 
+						htmlFile.WriteLine("'/>");
+						htmlFile.Write(lineStartTemplate.Replace("{LINECOLOUR}", colourString));
+						htmlFile.Write(lineCoordsStr);
+					}
+					oldx = x;
+					oldy = y;
+					lastFrameIndex = i;
 				}
-				string colourString = "#ff0000";
-				htmlFile.WriteLine("' style='fill:none;stroke-width:1.3; clip-path: url(#graphArea)' stroke='" + colourString + "'/>");
+				htmlFile.WriteLine("'/>");
 				float circleRadius = 3;
 				string eventColourString = "#ffffff";
 				foreach (MapOverlayEvent mapEvent in events)
@@ -1667,17 +1751,18 @@ namespace PerfSummaries
 					{
 						if (CsvStats.DoesSearchStringMatch(ev.Name, mapEvent.name))
 						{
+							string eventText = mapEvent.shortName != null ? mapEvent.shortName : ev.Name;
 							float x = xStat.samples[ev.Frame];
 							float y = yStat.samples[ev.Frame];
 							int svgX = toSvgX(x, y);
 							int svgY = toSvgY(x, y);
 							htmlFile.Write("<circle cx='" + svgX + "' cy='" + svgY + "' r='" + circleRadius + "' fill='" + eventColourString + "' fill-opacity='1.0'/>");
-							htmlFile.WriteLine("<text x='"+(svgX+5)+"' y='"+svgY+"' text-anchor='left' style='font-family: Verdana;fill: #ffffff; font-size: " + 9 + "px;'>" + ev.Name + "</text>");
+							htmlFile.WriteLine("<text x='"+(svgX+5)+"' y='"+svgY+"' text-anchor='left' style='font-family: Verdana;fill: #ffffff; font-size: " + 9 + "px;'>" + eventText + "</text>");
 						}
 					}
 				}
 
-				htmlFile.WriteLine("<text x='50%' y='" + (imageHeight * 0.05) + "' text-anchor='middle' style='font-family: Verdana;fill: #FFFFFF; stroke: #C0C0C0;  font-size: " + 20 + "px;'>" + title + "</text>");
+				//htmlFile.WriteLine("<text x='50%' y='" + (imageHeight * 0.05) + "' text-anchor='middle' style='font-family: Verdana;fill: #FFFFFF; stroke: #C0C0C0;  font-size: " + 20 + "px;'>" + title + "</text>");
 				htmlFile.WriteLine("</svg>");
 			}
 
@@ -1700,8 +1785,11 @@ namespace PerfSummaries
 		string destImageFilename;
 		float imageWidth;
 		float imageHeight;
+		float lineSplitDistanceThreshold;
 		int framesPerLineSegment;
 
+		MapOverlayEvent startEvent;
+
 		List<MapOverlayEvent> events = new List<MapOverlayEvent>();
 	};
 
