PerfReportTool 3.93
- Support for multiple map overlays within a single MapOverlaySummary
- BoundedStatValues summaries no longer need to be bounded (begin/end events are optional)
- texture summary

diff --git a/Engine/Source/Programs/CSVTools/PerfReportTool/PerfReportTool.cs b/Engine/Source/Programs/CSVTools/PerfReportTool/PerfReportTool.cs
--- a/Engine/Source/Programs/CSVTools/PerfReportTool/PerfReportTool.cs
+++ b/Engine/Source/Programs/CSVTools/PerfReportTool/PerfReportTool.cs
@@ -19,7 +19,7 @@ namespace PerfReportTool
 {
     class Version
     {
-        private static string VersionString = "3.92";
+        private static string VersionString = "3.93";
 
         public static string Get() { return VersionString; }
     };
@@ -437,7 +437,7 @@ namespace PerfReportTool
 				{
 					if (bBulkMode)
 					{
-						Console.Out.WriteLine(e.Message);
+						Console.Out.WriteLine("[ERROR] : "+ e.Message);
 					}
 					else
 					{
diff --git a/Engine/Source/Programs/CSVTools/PerfReportTool/perfsummaries.cs b/Engine/Source/Programs/CSVTools/PerfReportTool/perfsummaries.cs
--- a/Engine/Source/Programs/CSVTools/PerfReportTool/perfsummaries.cs
+++ b/Engine/Source/Programs/CSVTools/PerfReportTool/perfsummaries.cs
@@ -1319,8 +1319,8 @@ namespace PerfSummaries
 			}
 
 			title = element.GetSafeAttibute("title", "Events");
-			beginEvent = element.Attribute("beginevent").Value;
-			endEvent = element.Attribute("endevent").Value;
+			beginEvent = element.GetSafeAttibute<string>("beginevent");
+			endEvent = element.GetSafeAttibute<string>("endevent");
 
 			endOffsetPercentage = 0.0;
 			XAttribute endOffsetAtt = element.Attribute("endoffsetpercent");
@@ -1371,39 +1371,50 @@ namespace PerfSummaries
 			int endFrame = int.MaxValue;
 
 			// Find the start and end frames based on the events
-			foreach (CsvEvent ev in csvStats.Events)
+			if (beginEvent != null)
 			{
-				if (CsvStats.DoesSearchStringMatch(ev.Name, beginEvent))
+				foreach (CsvEvent ev in csvStats.Events)
 				{
-					startFrame = Math.Max(ev.Frame,startFrame);
+					if (CsvStats.DoesSearchStringMatch(ev.Name, beginEvent))
+					{
+						startFrame = Math.Max(ev.Frame, startFrame);
+					}
 				}
-				if (CsvStats.DoesSearchStringMatch(ev.Name, endEvent))
+				if (startFrame == -1)
 				{
-					endFrame = Math.Min(ev.Frame, endFrame);
+					Console.WriteLine("BoundedStatValuesSummary: Begin event " + beginEvent + " was not found");
+					return;
 				}
 			}
-			if ( startFrame == -1 )
-			{
-				Console.WriteLine("BoundedStatValuesSummary: Begin event " + beginEvent+" was not found");
-				return;
-			}
-			if (endFrame==int.MaxValue)
+			if (endEvent != null)
 			{
-				Console.WriteLine("BoundedStatValuesSummary: End event " + endEvent + " was not found");
-				return;
+				foreach (CsvEvent ev in csvStats.Events)
+				{
+					if (CsvStats.DoesSearchStringMatch(ev.Name, endEvent))
+					{
+						endFrame = Math.Min(ev.Frame, endFrame);
+					}
+				}
+				if (endFrame == int.MaxValue)
+				{
+					Console.WriteLine("BoundedStatValuesSummary: End event " + endEvent + " was not found");
+					return;
+				}
 			}
 			if ( startFrame >= endFrame )
 			{
 				throw new Exception("BoundedStatValuesSummary: end event appeared before the start event");
 			}
-
+			endFrame = Math.Min(endFrame, csvStats.SampleCount - 1);
+			startFrame = Math.Max(startFrame, 0);
+			
+			// Adjust the end frame based on the specified offset percentage, but cache the old value (some columns may need the unmodified one)
 			int endEventFrame = Math.Min(csvStats.SampleCount, endFrame + 1);
 			if (endOffsetPercentage > 0.0)
 			{
 				double multiplier = endOffsetPercentage / 100.0;
 				endFrame += (int)((double)(endFrame-startFrame)*multiplier);
 			}
-
 			endFrame = Math.Min(csvStats.SampleCount, endFrame + 1);
 			StatSamples frameTimeStat = csvStats.GetStat("frametime");
 			List<float> frameTimes = frameTimeStat.samples;
@@ -1418,6 +1429,11 @@ namespace PerfSummaries
 				}
 			}
 
+			// Nothing to report, so bail out!
+			if (filteredColumns.Count == 0)
+			{
+				return;
+			}
 
 			// Process the column values
 			foreach (Column col in filteredColumns)
@@ -1573,11 +1589,49 @@ namespace PerfSummaries
 			{
 				name = inName;
 			}
+			public MapOverlayEvent(XElement element)
+			{
+			}
 			public string name;
-			public string metadataPrefix;
+			public string metadataKey;
 			public string shortName;
 			public string lineColor;
 		};
+
+		class MapOverlay
+		{
+			public MapOverlay(XElement element)
+			{
+				positionStatNames[0] = element.GetSafeAttibute<string>("xStat");
+				positionStatNames[1] = element.GetSafeAttibute<string>("yStat");
+				positionStatNames[2] = element.GetSafeAttibute<string>("zStat");
+				metadataPrefix = element.GetSafeAttibute<string>("metadataPrefix");
+				lineColor = element.GetSafeAttibute<string>("lineColor","#ffffff");
+				foreach (XElement eventEl in element.Elements("event"))
+				{
+					MapOverlayEvent ev = new MapOverlayEvent(eventEl.Attribute("name").Value);
+					ev.shortName = eventEl.GetSafeAttibute<string>("shortName");
+					ev.metadataKey = eventEl.GetSafeAttibute<string>("metadataKey");
+					ev.lineColor = eventEl.GetSafeAttibute<string>("lineColor");
+					if (eventEl.GetSafeAttibute<bool>("isStartEvent", false))
+					{
+						if (startEvent != null)
+						{
+							throw new Exception("Can't have multiple start events!");
+						}
+						startEvent = ev;
+					}
+					events.Add(ev);
+				}
+
+			}
+			public string [] positionStatNames = new string[3];
+			public string metadataPrefix;
+			public MapOverlayEvent startEvent;
+			public string lineColor;
+			public List<MapOverlayEvent> events = new List<MapOverlayEvent>();
+		}
+
 		public MapOverlaySummary(XElement element, string baseXmlDirectory)
 		{
 			ReadStatsFromXML(element);
@@ -1586,12 +1640,6 @@ namespace PerfSummaries
 				throw new Exception("<stats> element is not supported");
 			}
 
-			xStatName = element.GetSafeAttibute<string>("xStat");
-			yStatName = element.GetSafeAttibute<string>("yStat");
-			zStatName = element.GetSafeAttibute<string>("zStat");
-			stats.Add(xStatName);
-			stats.Add(yStatName);
-			stats.Add(zStatName);
 			sourceImagePath = element.GetSafeAttibute<string>("sourceImage");
 			if ( !System.IO.Path.IsPathRooted(sourceImagePath))
 			{
@@ -1608,21 +1656,13 @@ namespace PerfSummaries
 			framesPerLineSegment = element.GetSafeAttibute<int>("framesPerLineSegment", 5);
 			lineSplitDistanceThreshold = element.GetSafeAttibute<float>("lineSplitDistanceThreshold", float.MaxValue);
 
-			foreach (XElement eventEl in element.Elements("event"))
+			foreach (XElement overlayEl in element.Elements("overlay"))
 			{
-				MapOverlayEvent ev = new MapOverlayEvent(eventEl.Attribute("name").Value);
-				ev.shortName = eventEl.GetSafeAttibute<string>("shortName");
-				ev.metadataPrefix = eventEl.GetSafeAttibute<string>("metadataPrefix");
-				ev.lineColor = eventEl.GetSafeAttibute<string>("lineColor");
-				if (eventEl.GetSafeAttibute<bool>("isStartEvent",false))
-				{
-					if (startEvent != null)
-					{
-						throw new Exception("Can't have multiple start events!");
-					}
-					startEvent = ev;
-				}
-				events.Add(ev);
+				MapOverlay overlay = new MapOverlay(overlayEl);
+				overlays.Add(overlay);
+				stats.Add(overlay.positionStatNames[0]);
+				stats.Add(overlay.positionStatNames[1]);
+				stats.Add(overlay.positionStatNames[2]);
 			}
 		}
 
@@ -1653,111 +1693,120 @@ namespace PerfSummaries
 					System.IO.File.Copy(sourceImagePath, outputMapFilename);
 				}
 
-				StatSamples xStat = csvStats.GetStat(xStatName);
-				StatSamples yStat = csvStats.GetStat(yStatName);
-
-				if ( xStat == null || yStat == null )
-				{
-					throw new Exception("MapOverlaySummary: CSV doesn't contain stat " + xStatName + " and/or " + yStatName);
-				}
-
-				// If a startevent is specified, update the start frame
-				int startFrame = 0;
-				if (startEvent != null)
-				{
-					foreach (CsvEvent ev in csvStats.Events)
-					{
-						if (CsvStats.DoesSearchStringMatch(ev.Name, startEvent.name))
-						{
-							startFrame = ev.Frame;
-							break;
-						}
-					}
-				}
-
 				// Check if the file exists in the output directory
 				htmlFile.WriteLine("  <br><h3>" + title + "</h3>");
 				htmlFile.WriteLine("<svg version='1.1' xmlns='http://www.w3.org/2000/svg' xmlns:xlink='http://www.w3.org/1999/xlink' width='" + imageWidth + "' height='" + imageHeight + "'>");
 				htmlFile.WriteLine("<image href='" + outputMapFilename + "' width='" + imageWidth + "' height='" + imageHeight + "' />");
 
-				// Make a mapping from frame to map indices
-				List<KeyValuePair<int, MapOverlayEvent>> frameEvents = new List<KeyValuePair<int, MapOverlayEvent>>();
-				foreach (MapOverlayEvent mapEvent in events)
+				// Draw the overlays
+				foreach (MapOverlay overlay in overlays)
 				{
-					foreach (CsvEvent ev in csvStats.Events)
+					StatSamples xStat = csvStats.GetStat(overlay.positionStatNames[0]);
+					StatSamples yStat = csvStats.GetStat(overlay.positionStatNames[1]);
+
+					if (xStat == null || yStat == null)
 					{
-						if (CsvStats.DoesSearchStringMatch(ev.Name, mapEvent.name))
-						{
-							frameEvents.Add(new KeyValuePair<int, MapOverlayEvent>(ev.Frame, mapEvent));
-						}
+						throw new Exception("MapOverlaySummary: CSV doesn't contain stat " + overlay.positionStatNames[0] + " and/or " + overlay.positionStatNames[1]);
 					}
-				}
-				frameEvents.Sort((pair0, pair1) => pair0.Key.CompareTo(pair1.Key));
-				int eventIndex = 0;
-
-				//List<string> polyLineStrings = new List<string>();
-				string colourString = "#ff0000";
-				string lineStartTemplate = "<polyline style='fill:none;stroke-width:1.3;stroke:{LINECOLOUR}' points='";
-				htmlFile.Write(lineStartTemplate.Replace("{LINECOLOUR}", colourString));
-				float adjustedLineSplitDistanceThreshold = lineSplitDistanceThreshold * framesPerLineSegment;
-				float oldx=0;
-				float oldy = 0;
-				int lastFrameIndex = 0;
-				for ( int i=startFrame; i<xStat.samples.Count;i+=framesPerLineSegment)
-				{
-					float x = xStat.samples[i];
-					float y = yStat.samples[i];
-					string lineCoordsStr = toSvgX(x, y) + "," + toSvgY(x, y) + " ";
 
-					// Figure out which event we're up to 
-					bool restartLineStrip = false;
-					while (eventIndex < frameEvents.Count && lastFrameIndex < frameEvents[eventIndex].Key && i>= frameEvents[eventIndex].Key)
+					// If a startevent is specified, update the start frame
+					int startFrame = 0;
+					if (overlay.startEvent != null)
 					{
-						MapOverlayEvent mapEvent = frameEvents[eventIndex].Value;
-						if (mapEvent.lineColor != null)
+						foreach (CsvEvent ev in csvStats.Events)
 						{
-							colourString = mapEvent.lineColor;
-							restartLineStrip = true;
+							if (CsvStats.DoesSearchStringMatch(ev.Name, overlay.startEvent.name))
+							{
+								startFrame = ev.Frame;
+								break;
+							}
 						}
-						eventIndex++;
 					}
 
-					float maxManhattanDist = Math.Max(Math.Abs(x - oldx), Math.Abs(y - oldy));
-					if (maxManhattanDist > adjustedLineSplitDistanceThreshold)
+					// Make a mapping from frame to map indices
+					List<KeyValuePair<int, MapOverlayEvent>> frameEvents = new List<KeyValuePair<int, MapOverlayEvent>>();
+					foreach (MapOverlayEvent mapEvent in overlay.events)
 					{
-						restartLineStrip = true;
+						foreach (CsvEvent ev in csvStats.Events)
+						{
+							if (CsvStats.DoesSearchStringMatch(ev.Name, mapEvent.name))
+							{
+								frameEvents.Add(new KeyValuePair<int, MapOverlayEvent>(ev.Frame, mapEvent));
+							}
+						}
 					}
-					else
+					frameEvents.Sort((pair0, pair1) => pair0.Key.CompareTo(pair1.Key));
+					int eventIndex = 0;
+
+					// Draw the lines
+					string currentLineColor = overlay.lineColor;
+					string lineStartTemplate = "<polyline style='fill:none;stroke-width:1.3;stroke:{LINECOLOUR}' points='";
+					htmlFile.Write(lineStartTemplate.Replace("{LINECOLOUR}", currentLineColor));
+					float adjustedLineSplitDistanceThreshold = lineSplitDistanceThreshold * framesPerLineSegment;
+					float oldx = 0;
+					float oldy = 0;
+					int lastFrameIndex = 0;
+					for (int i = startFrame; i < xStat.samples.Count; i += framesPerLineSegment)
 					{
-						htmlFile.Write(lineCoordsStr);
-					}
+						float x = xStat.samples[i];
+						float y = yStat.samples[i];
+						string lineCoordsStr = toSvgX(x, y) + "," + toSvgY(x, y) + " ";
 
-					if ( restartLineStrip )
-					{ 
-						htmlFile.WriteLine("'/>");
-						htmlFile.Write(lineStartTemplate.Replace("{LINECOLOUR}", colourString));
-						htmlFile.Write(lineCoordsStr);
+						// Figure out which event we're up to so we can do color changes
+						bool restartLineStrip = false;
+						while (eventIndex < frameEvents.Count && lastFrameIndex < frameEvents[eventIndex].Key && i >= frameEvents[eventIndex].Key)
+						{
+							MapOverlayEvent mapEvent = frameEvents[eventIndex].Value;
+							string newLineColor = mapEvent.lineColor != null ? mapEvent.lineColor : overlay.lineColor;
+							// If we changed color, restart the line strip
+							if (newLineColor != currentLineColor)
+							{
+								currentLineColor = newLineColor;
+								restartLineStrip = true;
+							}
+							eventIndex++;
+						}
+
+						// If the distance between this point and the last is over the threshold, restart the line strip
+						float maxManhattanDist = Math.Max(Math.Abs(x - oldx), Math.Abs(y - oldy));
+						if (maxManhattanDist > adjustedLineSplitDistanceThreshold)
+						{
+							restartLineStrip = true;
+						}
+						else
+						{
+							htmlFile.Write(lineCoordsStr);
+						}
+
+						if (restartLineStrip)
+						{
+							htmlFile.WriteLine("'/>");
+							htmlFile.Write(lineStartTemplate.Replace("{LINECOLOUR}", currentLineColor));
+							htmlFile.Write(lineCoordsStr);
+						}
+						oldx = x;
+						oldy = y;
+						lastFrameIndex = i;
 					}
-					oldx = x;
-					oldy = y;
-					lastFrameIndex = i;
-				}
-				htmlFile.WriteLine("'/>");
-				float circleRadius = 3;
-				string eventColourString = "#ffffff";
-				foreach (MapOverlayEvent mapEvent in events)
-				{
-					foreach (CsvEvent ev in csvStats.Events)
+					htmlFile.WriteLine("'/>");
+
+					// Plot the events 
+					float circleRadius = 3;
+					string eventColourString = "#ffffff";
+					foreach (MapOverlayEvent mapEvent in overlay.events)
 					{
-						if (CsvStats.DoesSearchStringMatch(ev.Name, mapEvent.name))
+						foreach (CsvEvent ev in csvStats.Events)
 						{
-							string eventText = mapEvent.shortName != null ? mapEvent.shortName : ev.Name;
-							float x = xStat.samples[ev.Frame];
-							float y = yStat.samples[ev.Frame];
-							int svgX = toSvgX(x, y);
-							int svgY = toSvgY(x, y);
-							htmlFile.Write("<circle cx='" + svgX + "' cy='" + svgY + "' r='" + circleRadius + "' fill='" + eventColourString + "' fill-opacity='1.0'/>");
-							htmlFile.WriteLine("<text x='"+(svgX+5)+"' y='"+svgY+"' text-anchor='left' style='font-family: Verdana;fill: #ffffff; font-size: " + 9 + "px;'>" + eventText + "</text>");
+							if (CsvStats.DoesSearchStringMatch(ev.Name, mapEvent.name))
+							{
+								string eventText = mapEvent.shortName != null ? mapEvent.shortName : ev.Name;
+								float x = xStat.samples[ev.Frame];
+								float y = yStat.samples[ev.Frame];
+								int svgX = toSvgX(x, y);
+								int svgY = toSvgY(x, y);
+								htmlFile.Write("<circle cx='" + svgX + "' cy='" + svgY + "' r='" + circleRadius + "' fill='" + eventColourString + "' fill-opacity='1.0'/>");
+								htmlFile.WriteLine("<text x='" + (svgX + 5) + "' y='" + svgY + "' text-anchor='left' style='font-family: Verdana;fill: #ffffff; font-size: " + 9 + "px;'>" + eventText + "</text>");
+							}
 						}
 					}
 				}
@@ -1775,9 +1824,6 @@ namespace PerfSummaries
 		{
 		}
 		string title;
-		string xStatName;
-		string yStatName;
-		string zStatName;
 		string sourceImagePath;
 		float offsetX;
 		float offsetY;
@@ -1788,9 +1834,7 @@ namespace PerfSummaries
 		float lineSplitDistanceThreshold;
 		int framesPerLineSegment;
 
-		MapOverlayEvent startEvent;
-
-		List<MapOverlayEvent> events = new List<MapOverlayEvent>();
+		List<MapOverlay> overlays = new List<MapOverlay>();
 	};
 
 	class SummaryMetadataValue
@@ -1846,7 +1890,14 @@ namespace PerfSummaries
                 metadataValue = new SummaryMetadataValue(value, colorThresholdList, tooltip);
             }
 
-            dict.Add(key, metadataValue);
+			try
+			{
+				dict.Add(key, metadataValue);
+			}
+			catch (System.ArgumentException)
+			{
+				throw new Exception("Summary metadata key " + key + " has already been added");
+			}
 		}
 
         public void AddString(string key, string value, ColourThresholdList colorThresholdList = null, string tooltip = "")
