Perfreporttool 3.94
- Fixed exception with certain CSVs, caused by long-standing hack for MVP calculation which removed the last sample from the frametime stat.

diff --git a/Engine/Source/Programs/CSVTools/PerfReportTool/PerfReportTool.cs b/Engine/Source/Programs/CSVTools/PerfReportTool/PerfReportTool.cs
--- a/Engine/Source/Programs/CSVTools/PerfReportTool/PerfReportTool.cs
+++ b/Engine/Source/Programs/CSVTools/PerfReportTool/PerfReportTool.cs
@@ -19,7 +19,7 @@ namespace PerfReportTool
 {
     class Version
     {
-        private static string VersionString = "3.93";
+        private static string VersionString = "3.94";
 
         public static string Get() { return VersionString; }
     };
diff --git a/Engine/Source/Programs/CSVTools/PerfReportTool/perfsummaries.cs b/Engine/Source/Programs/CSVTools/PerfReportTool/perfsummaries.cs
--- a/Engine/Source/Programs/CSVTools/PerfReportTool/perfsummaries.cs
+++ b/Engine/Source/Programs/CSVTools/PerfReportTool/perfsummaries.cs
@@ -419,12 +419,14 @@ namespace PerfSummaries
 			public float TotalTimeSeconds;
 		};
 
-		FpsChartData ComputeFPSChartDataForFrames(List<float> frameTimes)
+		FpsChartData ComputeFPSChartDataForFrames(List<float> frameTimes, bool skiplastFrame)
 		{
 			double totalFrametime = 0.0;
 			int hitchCount = 0;
 			double totalHitchTime = 0.0;
 
+			int frameCount = skiplastFrame ? frameTimes.Count - 1 : frameTimes.Count;
+
 			// Count hitches
 			if (bUseEngineHitchMetric)
 			{
@@ -434,8 +436,9 @@ namespace PerfSummaries
 				double LastFrameTime = float.MinValue;
 				float HitchMultiplierAmount = GetEngineHitchToNonHitchRatio();
 
-				foreach (float frametime in frameTimes)
+				for ( int i=0; i< frameCount; i++)
 				{
+					float frametime = frameTimes[i];
 					// How long has it been since the last hitch we detected?
 					if (frametime >= hitchThreshold)
 					{
@@ -461,8 +464,9 @@ namespace PerfSummaries
 			}
 			else
 			{
-				foreach (float frametime in frameTimes)
+				for (int i = 0; i < frameCount; i++)
 				{
+					float frametime = frameTimes[i];
 					totalFrametime += frametime;
 					if (frametime >= hitchThreshold)
 					{
@@ -494,13 +498,9 @@ namespace PerfSummaries
                 string csvPath = Path.Combine(Path.GetDirectoryName(htmlFileName), "FrameStats_colored.csv");
                 statsCsvFile = new System.IO.StreamWriter(csvPath, false);
             }
-            // Compute MVP30 and MVP60
+			// Compute MVP30 and MVP60. Note: we ignore the last frame because fpscharts can hitch
             List<float> frameTimes = csvStats.Stats["frametime"].samples;
-
-			// Remove the last element, because FPSCharts can hitch
-			frameTimes.RemoveAt(frameTimes.Count - 1);
-
-			FpsChartData fpsChartData = ComputeFPSChartDataForFrames(frameTimes);
+			FpsChartData fpsChartData = ComputeFPSChartDataForFrames(frameTimes,true);
 
             // Write the averages
             List<string> ColumnNames = new List<string>();
@@ -625,7 +625,7 @@ namespace PerfSummaries
                     {
                         continue;
                     }
-					FpsChartData captureFpsChartData = ComputeFPSChartDataForFrames(CaptureFrameTimes.Frames);
+					FpsChartData captureFpsChartData = ComputeFPSChartDataForFrames(CaptureFrameTimes.Frames,true);
 				
                     if (captureFpsChartData.TotalTimeSeconds == 0.0f)
                     {
