PerfReportTool 3.95:
- Fix issue where multiple begin/end events with the same name in the BoundedStatSummary throws an exception

diff --git a/Engine/Source/Programs/CSVTools/PerfReportTool/PerfReportTool.cs b/Engine/Source/Programs/CSVTools/PerfReportTool/PerfReportTool.cs
--- a/Engine/Source/Programs/CSVTools/PerfReportTool/PerfReportTool.cs
+++ b/Engine/Source/Programs/CSVTools/PerfReportTool/PerfReportTool.cs
@@ -19,7 +19,7 @@ namespace PerfReportTool
 {
     class Version
     {
-        private static string VersionString = "3.94";
+        private static string VersionString = "3.95";
 
         public static string Get() { return VersionString; }
     };
diff --git a/Engine/Source/Programs/CSVTools/PerfReportTool/perfsummaries.cs b/Engine/Source/Programs/CSVTools/PerfReportTool/perfsummaries.cs
--- a/Engine/Source/Programs/CSVTools/PerfReportTool/perfsummaries.cs
+++ b/Engine/Source/Programs/CSVTools/PerfReportTool/perfsummaries.cs
@@ -1377,7 +1377,8 @@ namespace PerfSummaries
 				{
 					if (CsvStats.DoesSearchStringMatch(ev.Name, beginEvent))
 					{
-						startFrame = Math.Max(ev.Frame, startFrame);
+						startFrame = ev.Frame;
+						break;
 					}
 				}
 				if (startFrame == -1)
@@ -1392,7 +1393,11 @@ namespace PerfSummaries
 				{
 					if (CsvStats.DoesSearchStringMatch(ev.Name, endEvent))
 					{
-						endFrame = Math.Min(ev.Frame, endFrame);
+						endFrame = ev.Frame;
+						if ( endFrame > startFrame )
+						{
+							break;
+						}
 					}
 				}
 				if (endFrame == int.MaxValue)
