PerfReportTool 3.96
- Fix issue with NaNs getting output to the email table in some cases

diff --git a/Engine/Source/Programs/CSVTools/PerfReportTool/PerfReportTool.cs b/Engine/Source/Programs/CSVTools/PerfReportTool/PerfReportTool.cs
--- a/Engine/Source/Programs/CSVTools/PerfReportTool/PerfReportTool.cs
+++ b/Engine/Source/Programs/CSVTools/PerfReportTool/PerfReportTool.cs
@@ -19,7 +19,7 @@ namespace PerfReportTool
 {
     class Version
     {
-        private static string VersionString = "3.95";
+        private static string VersionString = "3.96";
 
         public static string Get() { return VersionString; }
     };
diff --git a/Engine/Source/Programs/CSVTools/PerfReportTool/perfsummaries.cs b/Engine/Source/Programs/CSVTools/PerfReportTool/perfsummaries.cs
--- a/Engine/Source/Programs/CSVTools/PerfReportTool/perfsummaries.cs
+++ b/Engine/Source/Programs/CSVTools/PerfReportTool/perfsummaries.cs
@@ -2185,7 +2185,8 @@ namespace PerfSummaries
 			List<float> RowMaxValues = new List<float>();
 			List<float> RowTotals = new List<float>();
 			List<float> RowMinValues = new List<float>();
-			List<int> RowValueCounts = new List<int>();
+			List<int> RowCounts = new List<int>();
+			List<ColourThresholdList> RowColourThresholds = new List<ColourThresholdList>();
 
 			// Set the initial sort key
 			string CurrentRowSortKey = "";
@@ -2204,13 +2205,15 @@ namespace PerfSummaries
 					RowMaxValues.Clear();
 					RowMinValues.Clear();
 					RowTotals.Clear();
-					RowValueCounts.Clear();
+					RowCounts.Clear();
+					RowColourThresholds.Clear();
 					for (int j = 0; j < columns.Count; j++)
 					{
 						RowMaxValues.Add(-float.MaxValue);
 						RowMinValues.Add(float.MaxValue);
 						RowTotals.Add(0.0f);
-						RowValueCounts.Add(0);
+						RowCounts.Add(0);
+						RowColourThresholds.Add(null);
 					}
 					mergedRowsCount = 0;
 					reset = false;
@@ -2223,16 +2226,13 @@ namespace PerfSummaries
 					if (column.isNumeric)
 					{
 						float value = column.GetValue(i);
-						if (value == float.MaxValue)
-						{
-							//RowTotals[j] = float.MaxValue;
-						}
-						else
+						if (value != float.MaxValue)
 						{
 							RowMaxValues[j] = Math.Max(RowMaxValues[j], value);
 							RowMinValues[j] = Math.Min(RowMinValues[j], value);
 							RowTotals[j] += value;
-							RowValueCounts[j]++;
+							RowColourThresholds[j] = column.GetColourThresholds(i);
+							RowCounts[j]++;
 						}
 					}
 				}
@@ -2260,16 +2260,15 @@ namespace PerfSummaries
 					newColumns[countColumnIndex].SetValue(destRowIndex, (float)mergedRowsCount);
 					for (int j = 0; j < columns.Count; j++)
 					{
-						SummaryMetadataColumn srcColumn = columns[j];
 						int destColumnBaseIndex = srcToDestBaseColumnIndex[j];
-						if (destColumnBaseIndex != -1 && RowTotals[j] != float.MaxValue)
+						if (destColumnBaseIndex != -1 && RowCounts[j]>0)
 						{
-							newColumns[destColumnBaseIndex].SetValue(destRowIndex, RowTotals[j] / (float)RowValueCounts[j]);
+							newColumns[destColumnBaseIndex].SetValue(destRowIndex, RowTotals[j] / (float)RowCounts[j]);
 							newColumns[destColumnBaseIndex+1].SetValue(destRowIndex, RowMinValues[j]);
 							newColumns[destColumnBaseIndex+2].SetValue(destRowIndex, RowMaxValues[j]);
 
 							// Set colour thresholds based on the source column
-							ColourThresholdList Thresholds = srcColumn.GetColourThresholds(i);
+							ColourThresholdList Thresholds = RowColourThresholds[j];
 							for ( int k=0;k<3;k++)
 							{
 								newColumns[destColumnBaseIndex+k].SetColourThresholds(destRowIndex, Thresholds);
