PerfReportTool 3.98
- Fix image path issue

diff --git a/Engine/Source/Programs/CSVTools/PerfReportTool/PerfReportTool.cs b/Engine/Source/Programs/CSVTools/PerfReportTool/PerfReportTool.cs
--- a/Engine/Source/Programs/CSVTools/PerfReportTool/PerfReportTool.cs
+++ b/Engine/Source/Programs/CSVTools/PerfReportTool/PerfReportTool.cs
@@ -19,7 +19,7 @@ namespace PerfReportTool
 {
     class Version
     {
-        private static string VersionString = "3.96";
+        private static string VersionString = "3.98";
 
         public static string Get() { return VersionString; }
     };
diff --git a/Engine/Source/Programs/CSVTools/PerfReportTool/perfsummaries.cs b/Engine/Source/Programs/CSVTools/PerfReportTool/perfsummaries.cs
--- a/Engine/Source/Programs/CSVTools/PerfReportTool/perfsummaries.cs
+++ b/Engine/Source/Programs/CSVTools/PerfReportTool/perfsummaries.cs
@@ -1701,7 +1701,7 @@ namespace PerfSummaries
 				// Check if the file exists in the output directory
 				htmlFile.WriteLine("  <br><h3>" + title + "</h3>");
 				htmlFile.WriteLine("<svg version='1.1' xmlns='http://www.w3.org/2000/svg' xmlns:xlink='http://www.w3.org/1999/xlink' width='" + imageWidth + "' height='" + imageHeight + "'>");
-				htmlFile.WriteLine("<image href='" + outputMapFilename + "' width='" + imageWidth + "' height='" + imageHeight + "' />");
+				htmlFile.WriteLine("<image href='" + destImageFilename + "' width='" + imageWidth + "' height='" + imageHeight + "' />");
 
 				// Draw the overlays
 				foreach (MapOverlay overlay in overlays)
