CsvToSVG 2.3
 - Batched + multithreaded support via response files

PerfReportTool 4.00
 - Support batched graph generation. Increases perf by 30% over old multi process method. Significantly lower CPU and disk usage too. Enable with -batchedgraphs
 - Fixed determinism issues with report generation. Test cases are now identical between runs

diff --git a/Engine/Source/Programs/CSVTools/CSVToSVG/CsvToSVG.cs b/Engine/Source/Programs/CSVTools/CSVToSVG/CsvToSVG.cs
--- a/Engine/Source/Programs/CSVTools/CSVToSVG/CsvToSVG.cs
+++ b/Engine/Source/Programs/CSVTools/CSVToSVG/CsvToSVG.cs
@@ -10,6 +10,7 @@ using System.Drawing;
 using System.IO;
 using System.Globalization;
 using System.Diagnostics;
+using System.Security.Cryptography;
 using CSVStats;
 
 // 2.1 TODO
@@ -24,7 +25,7 @@ namespace CSVTools
 {
     class Version
     {
-        private static string VersionString = "2.21";
+        private static string VersionString = "2.3";
         
         public static string Get() { return VersionString; }
     };
@@ -154,7 +155,31 @@ namespace CSVTools
         public Colour[] GraphColours;
     };
 
-    class Program : CSVStats.CommandLineTool
+	class CsvCache
+	{
+		public CsvStats ReadCSVFile(string csvFilename, string[] statNames, int numRowsToSkip)
+		{
+			lock (readLock)
+			{
+				if (csvFilename != filename || skipRows != numRowsToSkip)
+				{
+					csvStats = CsvStats.ReadCSVFile(csvFilename, null, numRowsToSkip);
+					filename = csvFilename;
+					skipRows = numRowsToSkip;
+				}
+				// Copy just the stats we're interested in from the cached file
+				CsvStats outStats = new CsvStats(csvStats, statNames);
+				return outStats;
+			}
+		}
+		public CsvStats csvStats;
+		public string filename;
+		public int skipRows;
+		private readonly object readLock = new object();
+	};
+
+
+	class Program : CSVStats.CommandLineTool
     {
         System.IO.StreamWriter svgFile;
         Rect dimensions = new Rect(0, 0, 1000, 500);
@@ -169,17 +194,19 @@ namespace CSVTools
         int maxHierarchyDepth = -1;
         char hierarchySeparator = '/';
         int colourOffset = 0;
-        string UniqueID = Guid.NewGuid().ToString().Substring(26);
 
-        static string formatString =
+		static string formatString =
             "Format: \n" +
             "       -csvs <list> OR -csv <list> OR -csvDir <path>\n" +
             "       [ -o <svgFilename> ]\n" +
             "       -stats <stat names> (can include wildcards)\n" +
-            "     OR \n" +
-            "       -updatesvg <svgFilename>\n" +
-            "       NOTE: this updates an svg by regenerating it with the original commandline parameters - requires original data\n\n" +
-            "\nOptional Args: \n" +
+			"     OR \n" +
+			"       -batchCommands <response file with commandlines>\n" +
+			"       -mt <number of threads> \n" +
+			"     OR \n" +
+			"       -updatesvg <svgFilename>\n" +
+			"       NOTE: this updates an svg by regenerating it with the original commandline parameters - requires original data\n\n" +
+			"\nOptional Args: \n" +
 			"       -averageThreshold <value>\n" +
 			"       -budget <ms>\n" +
 			"       -colourOffset <value>\n" +
@@ -222,46 +249,65 @@ namespace CSVTools
             "       -percentile \n" +
             "       -percentile90 \n" +
             "       -percentile99 \n" +
+			"       -uniqueId <string> : unique ID for JS (needed if this is getting embedded in HTML alongside other graphs)\n" +
+			"       -nocommandlineEmbed : don't embed the commandline in the SVG" +
 			"       -lineDecimalPlaces <N> (default 3)" +
             "";
 
+		void Run(string[] args)
+		{
+			if (args.Length < 2)
+			{
+				WriteLine("CsvToSVG " + Version.Get());
+				WriteLine(formatString);
+				return;
+			}
 
-        void Run( string[] args )
-        {
-            if (args.Length < 2)
-            {
-                WriteLine("CsvToSVG " + Version.Get() );
-                WriteLine(formatString);
-                return;
-            }
+			// Read the command line
+			ReadCommandLine(args);
 
-            // Read the command line
-            ReadCommandLine(args);
+			string svgToUpdate = GetArg("updatesvg");
+			if (svgToUpdate.Length > 0)
+			{
+				if (args.Length > 2)
+				{
+					WriteLine("-UpdateSVG <svgFilename> must be the only argument!");
+					return;
+				}
+				string newCommandLine = "";
+				string[] svgLines = ReadLinesFromFile(svgToUpdate);
+				for (int i = 0; i < svgLines.Length - 2; i++)
+				{
+					if (svgLines[i].StartsWith("<![CDATA[") && svgLines[i + 1].StartsWith("Created with CSVtoSVG ") && svgLines[i + 1].EndsWith(" with commandline:"))
+					{
+						newCommandLine = svgLines[i + 2];
+						break;
+					}
+				}
 
-            string svgToUpdate = GetArg("updatesvg");
-            if (svgToUpdate.Length > 0)
-            {
-                if (args.Length > 2)
-                {
-                    WriteLine("-UpdateSVG <svgFilename> must be the only argument!");
-                    return;
-                }
-                string newCommandLine = "";
-                string [] svgLines = ReadLinesFromFile(svgToUpdate);
-                for ( int i=0; i<svgLines.Length-2; i++)
-                {
-                    if (svgLines[i].StartsWith("<![CDATA[") && svgLines[i+1].StartsWith("Created with CSVtoSVG ") && svgLines[i+1].EndsWith(" with commandline:"))
-                    {
-                        newCommandLine = svgLines[i + 2];
-                        break;
-                    }
-                }
-               
-                ReadCommandLine(MakeArgsArray(newCommandLine));
-            }
+				ReadCommandLine(MakeArgsArray(newCommandLine));
+			}
+
+			MakeGraph();
+		}
 
-            // Read CSV filenames from a directory or list
-            string[] csvFilenames;
+		string GenerateIDFromString(string str)
+		{
+			string id = "";
+			HashAlgorithm algorithm = SHA256.Create();
+			byte[] hash = algorithm.ComputeHash(Encoding.UTF8.GetBytes(commandLine.GetCommandLine()));
+			for ( int i=24; i<32; i++ )
+			{
+				id += hash[i].ToString("x2");
+			}
+			return id;
+		}
+
+
+		void MakeGraph()
+		{
+			// Read CSV filenames from a directory or list
+			string[] csvFilenames;
             string csvDir = GetArg("csvDir"); 
             if (csvDir.Length > 0)
             {
@@ -304,14 +350,9 @@ namespace CSVTools
                 WriteLine("Can't display multiple stats and multiple CSVs");
                 return;
             }
-            //if ( csvFilenames.Length > 1 && GetBoolArg("stacked"))
-            //{
-            //    WriteLine("Can't stack multiple CSVs");
-            //    return;
-            //}
 
-            // Figure out the filename, if it wasn't provided
-            string svgFilename = GetArg("o", false);
+			// Figure out the filename, if it wasn't provided
+			string svgFilename = GetArg("o", false);
             if (svgFilename.Length == 0)
             {
                 if (csvFilenames.Length == 1)
@@ -341,8 +382,8 @@ namespace CSVTools
                 WriteLine("Couldn't figure out an appropriate filename, and no filename was provided with -o");
             }
 
-
-            dimensions.width = GetIntArg("width", 1800);
+			// Generate a unique ID based on the commandline. This is needed in order to embed in HTML alongside other graphs
+			dimensions.width = GetIntArg("width", 1800);
             dimensions.height = GetIntArg("height", 550);
 
             svgFile = new System.IO.StreamWriter(svgFilename);
@@ -353,9 +394,16 @@ namespace CSVTools
             }
             SvgWriteLine(">");
 
-            SvgWriteLine("<![CDATA[ \nCreated with CSVtoSVG " + Version.Get() + " with commandline:");
-            SvgWriteLine(commandLine);
-            SvgWriteLine("]]>");
+			if (GetBoolArg("nocommandlineEmbed"))
+			{
+				SvgWriteLine("<![CDATA[ \nCreated with CSVtoSVG " + Version.Get() );
+			}
+			else
+			{
+				SvgWriteLine("<![CDATA[ \nCreated with CSVtoSVG " + Version.Get() + " with commandline:");
+				SvgWriteLine(commandLine.GetCommandLine());
+			}
+			SvgWriteLine("]]>");
 
             bool writeErrorsToSVG = GetIntArg("writeErrorsToSVG",1) == 1;
             try
@@ -390,11 +438,16 @@ namespace CSVTools
             }
         }
 
+		string GetGraphUniqueId()
+		{
+			return GetArg("uniqueID", "ID");
+		}
+
         void SvgWrite(string str, bool replaceUniqueNames=true)
         {
             if (replaceUniqueNames)
             {
-                svgFile.Write(str.Replace("<UNIQUE>", "_" + UniqueID));
+                svgFile.Write(str.Replace("<UNIQUE>", "_" + GetGraphUniqueId()));
             }
             else
             {
@@ -405,7 +458,7 @@ namespace CSVTools
         {
             if (replaceUniqueNames)
             {
-                svgFile.WriteLine(str.Replace("<UNIQUE>", "_" + UniqueID));
+                svgFile.WriteLine(str.Replace("<UNIQUE>", "_" + GetGraphUniqueId()));
             }
             else
             {
@@ -855,7 +908,7 @@ namespace CSVTools
             return false;
         }
 
-        CsvStats ProcessCSV(string csvFilename, string[] statNames, bool bDiscardLastFrame )
+		CsvStats ProcessCSV(string csvFilename, string[] statNames, bool bDiscardLastFrame )
         {
             int numRowsToSkip = 0;
             numRowsToSkip = GetIntArg("skipRows",0);
@@ -863,10 +916,18 @@ namespace CSVTools
             bool filterOutZeros = false;
             filterOutZeros = GetBoolArg("filterOutZeros");
 
-            CsvStats csvStats = CsvStats.ReadCSVFile(csvFilename, statNames, numRowsToSkip );
-
-            // Sometimes the last frame is garbage. We might want to remove it
-            if (bDiscardLastFrame)
+			CsvStats csvStats = null;
+			if (csvCache != null)
+			{
+				// Copy the data from the source CSV
+				csvStats = csvCache.ReadCSVFile(csvFilename, statNames, numRowsToSkip);
+			}
+			else
+			{
+				csvStats = CsvStats.ReadCSVFile(csvFilename, statNames, numRowsToSkip);
+			}
+			// Sometimes the last frame is garbage. We might want to remove it
+			if (bDiscardLastFrame)
             {
                 foreach (StatSamples stat in csvStats.Stats.Values.ToArray())
                 {
@@ -2386,24 +2447,57 @@ namespace CSVTools
             SvgWriteLine("<text x='" + 10 + "' y='" + y + "' fill=" + colour.SVGString() + " font-size='9' font-family='Courier New' >" + Commandline  + "</text>");
         }
 
-        static void Main(string[] args)
+
+		static CsvCache csvCache = null;
+
+		static void Main(string[] args)
         {
-            Program program = new Program();
-            if (Debugger.IsAttached)
-            {
-                program.Run(args);
-            }
-            else
-            {
-                try
-                {
-                    program.Run(args);
-                }
-                catch (System.Exception e)
-                {
-                    Console.WriteLine("[ERROR] " + e.Message);
-                }
-            }
+			CommandLine commandline = new CommandLine(args);
+			string batchCommandsFilename = commandline.GetArg("batchCommands");
+
+			if (batchCommandsFilename.Length > 0)
+			{
+				// In batch mode, just output a list of commandlines
+				string[] commandlines = System.IO.File.ReadAllLines(batchCommandsFilename);
+				if (commandlines.Length > 1)
+				{
+					// Enable caching
+					csvCache = new CsvCache();
+				}
+				int numThreads = commandline.GetIntArg("mt", 4);
+				if (numThreads>1)
+				{
+					var result = Parallel.For(0, commandlines.Length, new ParallelOptions { MaxDegreeOfParallelism = numThreads }, i =>
+					{
+						Program program = new Program();
+						program.Run(MakeArgsArray(commandlines[i]));
+					});
+				}
+				else
+				{
+					foreach (string cmdLine in commandlines)
+					{
+						Program program = new Program();
+						program.Run(MakeArgsArray(cmdLine));
+					}
+				}
+			}
+			else
+			{
+				Program program = new Program();
+				try
+				{
+					program.Run(args);
+				}
+				catch (System.Exception e)
+				{
+					Console.WriteLine("[ERROR] " + e.Message);
+					if (Debugger.IsAttached)
+					{
+						throw e;
+					}
+				}
+			}
         }
      }
 }
diff --git a/Engine/Source/Programs/CSVTools/CsvStats/CommandLineTool.cs b/Engine/Source/Programs/CSVTools/CsvStats/CommandLineTool.cs
--- a/Engine/Source/Programs/CSVTools/CsvStats/CommandLineTool.cs
+++ b/Engine/Source/Programs/CSVTools/CsvStats/CommandLineTool.cs
@@ -11,6 +11,120 @@ using System.Diagnostics;
 
 namespace CSVStats
 {
+	public class CommandLine
+	{
+		protected string commandLine = "";
+
+		protected Dictionary<string, string> CommandLineArgs;
+
+		public string GetCommandLine()
+		{
+			return commandLine;
+		}
+
+		public CommandLine(string[] args)
+		{
+			commandLine = "";
+			foreach (string arg in args)
+			{
+				if (arg.Contains(' ') || arg.Contains('\t'))
+				{
+					commandLine += "\"" + arg + "\" ";
+				}
+				else
+				{
+					commandLine += arg + " ";
+				}
+			}
+			CommandLineArgs = new Dictionary<string, string>();
+			for (int i = 0; i < args.Length; i++)
+			{
+				string arg = args[i];
+				if (arg[0] == '-')
+				{
+					string val = "1";
+
+					// If there's a value, read it
+					if (i < args.Length - 1 && args[i + 1][0] != '-')
+					{
+						bool first = true;
+						for (int j = i + 1; j < args.Length; j++)
+						{
+							string str = args[j];
+							if (str.Length > 0 && str[0] == '-')
+								break;
+							if (first)
+							{
+								val = str;
+								first = false;
+							}
+							else
+							{
+								val += ";" + str;
+							}
+						}
+						i++;
+					}
+
+					string argKey = arg.Substring(1).ToLower();
+					if (CommandLineArgs.ContainsKey(argKey))
+					{
+						Console.Out.WriteLine("Duplicate commandline argument found for " + arg + ". Overriding value from " + CommandLineArgs[argKey] + " to " + val);
+						CommandLineArgs.Remove(argKey);
+					}
+					CommandLineArgs.Add(arg.Substring(1).ToLower(), val);
+				}
+			}
+		}
+
+		public int GetIntArg(string key, int defaultValue)
+		{
+			string val = GetArg(key, false);
+			if (val != "") return Convert.ToInt32(val);
+			return defaultValue;
+		}
+
+		public float GetFloatArg(string key, float defaultValue)
+		{
+			string val = GetArg(key, false);
+			if (val != "") return (float)Convert.ToDouble(val, System.Globalization.CultureInfo.InvariantCulture);
+			return defaultValue;
+		}
+
+		public bool GetBoolArg(string key)
+		{
+			return CommandLineArgs.ContainsKey(key.ToLower());
+		}
+
+		public string GetArg(string key, string defaultValue)
+		{
+			string lowerKey = key.ToLower();
+
+			if (CommandLineArgs.ContainsKey(lowerKey))
+			{
+				return CommandLineArgs[lowerKey];
+			}
+			return defaultValue;
+		}
+
+		public string GetArg(string key, bool mandatory = false)
+		{
+			string lowerKey = key.ToLower();
+
+			if (CommandLineArgs.ContainsKey(lowerKey))
+			{
+				return CommandLineArgs[lowerKey];
+			}
+			else if (mandatory)
+			{
+				Console.WriteLine("Missing parameter " + key);
+			}
+
+			return "";
+		}
+
+	}
+
 	public class CommandLineTool
 	{
 		protected enum HostPlatform
@@ -20,9 +134,7 @@ namespace CSVStats
 			Linux
 		}
 
-		protected string commandLine = "";
-
-		protected Dictionary<string, string> CommandLineArgs;
+		protected CommandLine commandLine;
 
 		private readonly static bool bIsMac = File.Exists("/System/Library/CoreServices/SystemVersion.plist");
 
@@ -59,55 +171,33 @@ namespace CSVStats
         }
         protected int GetIntArg(string key, int defaultValue)
         {
-            string val = GetArg(key, false);
-            if (val != "") return Convert.ToInt32(val);
-            return defaultValue;
-        }
+			return commandLine.GetIntArg(key, defaultValue);
+		}
 
-        protected float GetFloatArg(string key, float defaultValue)
+		protected float GetFloatArg(string key, float defaultValue)
         {
-            string val = GetArg(key, false);
-            if (val != "") return (float)Convert.ToDouble(val,System.Globalization.CultureInfo.InvariantCulture);
-            return defaultValue;
-        }
+			return commandLine.GetFloatArg(key, defaultValue);
+		}
 
-        protected bool GetBoolArg(string key)
+		protected bool GetBoolArg(string key)
         {
-            return CommandLineArgs.ContainsKey(key.ToLower());
+			return commandLine.GetBoolArg(key);
         }
 
 		protected string GetArg(string key, string defaultValue)
 		{
-			string lowerKey = key.ToLower();
-
-			if (CommandLineArgs.ContainsKey(lowerKey))
-			{
-				return CommandLineArgs[lowerKey];
-			}
-			return defaultValue;
+			return commandLine.GetArg(key, defaultValue);
 		}
 
 		protected string GetArg(string key, bool mandatory = false)
         {
-            string lowerKey = key.ToLower();
-
-            if (CommandLineArgs.ContainsKey(lowerKey))
-            {
-                return CommandLineArgs[lowerKey];
-            }
-            else if (mandatory)
-            {
-                WriteLine("Missing parameter {0}", key);
-            }
-
-            return "";
+			return commandLine.GetArg(key, mandatory);
         }
 
 		protected void WriteLine(String message, params object[] args)
         {
             String formatted = String.Format(message, args);
             Console.WriteLine(formatted);
-            //            Trace.WriteLine(formatted);
         }
 
 
@@ -148,61 +238,9 @@ namespace CSVStats
             return lines.ToArray();
         }
 
-        protected void ReadCommandLine(string[] args)
-        {
-            commandLine = "";
-            foreach (string arg in args)
-            {
-                if (arg.Contains(' ') || arg.Contains('\t'))
-                {
-                    commandLine += "\"" + arg + "\" ";
-                }
-                else
-                {
-                    commandLine += arg + " ";
-                }
-            }
-            CommandLineArgs = new Dictionary<string, string>();
-            for (int i = 0; i < args.Length; i++)
-            {
-                string arg = args[i];
-                if (arg[0] == '-')
-                {
-                    string val = "1";
-
-                    // If there's a value, read it
-                    if (i < args.Length - 1 && args[i + 1][0] != '-')
-                    {
-                        bool first = true;
-                        for (int j = i + 1; j < args.Length; j++)
-                        {
-                            string str = args[j];
-                            if (str.Length > 0 && str[0] == '-')
-                                break;
-                            if (first)
-                            {
-                                val = str;
-                                first = false;
-                            }
-                            else
-                            {
-                                val += ";" + str;
-                            }
-                        }
-                        i++;
-                    }
-
-                    string argKey = arg.Substring(1).ToLower();
-                    if (CommandLineArgs.ContainsKey(argKey))
-                    {
-                        Console.Out.WriteLine("Duplicate commandline argument found for " + arg + ". Overriding value from " + CommandLineArgs[argKey] + " to " + val);
-                        CommandLineArgs.Remove(argKey);
-                    }
-                    CommandLineArgs.Add(arg.Substring(1).ToLower(), val);
-                }
-            }
-        }
-
-		
+		protected void ReadCommandLine(string[] args)
+		{
+			commandLine = new CommandLine(args);
+		}
 	}
 }
diff --git a/Engine/Source/Programs/CSVTools/CsvStats/CsvStats.cs b/Engine/Source/Programs/CSVTools/CsvStats/CsvStats.cs
--- a/Engine/Source/Programs/CSVTools/CsvStats/CsvStats.cs
+++ b/Engine/Source/Programs/CSVTools/CsvStats/CsvStats.cs
@@ -372,21 +372,39 @@ namespace CSVStats
             Stats = new Dictionary<string, StatSamples>();
             Events = new List<CsvEvent>();
         }
-		public CsvStats(CsvStats source)
+		public CsvStats(CsvStats sourceCsvStats, string [] statNamesToFilter=null)
 		{
 			Stats = new Dictionary<string, StatSamples>();
 			Events = new List<CsvEvent>();
-			foreach (StatSamples sourceStat in source.Stats.Values)
+			if (statNamesToFilter != null)
 			{
-				AddStat(new StatSamples(sourceStat));
+				foreach (string statName in statNamesToFilter)
+				{
+					List<StatSamples> stats = sourceCsvStats.GetStatsMatchingString(statName);
+					foreach (StatSamples sourceStat in stats)
+					{
+						string key = sourceStat.Name.ToLower();
+						if (!Stats.ContainsKey(key))
+						{
+							Stats.Add(key,new StatSamples(sourceStat));
+						}
+					}
+				}
+			}
+			else
+			{
+				foreach (StatSamples sourceStat in sourceCsvStats.Stats.Values)
+				{
+					AddStat(new StatSamples(sourceStat));
+				}
 			}
-			foreach (CsvEvent ev in source.Events)
+			foreach (CsvEvent ev in sourceCsvStats.Events)
 			{
 				Events.Add(new CsvEvent(ev));
 			}
-			if (source.metaData != null)
+			if (sourceCsvStats.metaData != null)
 			{
-				metaData = new CsvMetadata(source.metaData);
+				metaData = new CsvMetadata(sourceCsvStats.metaData);
 			}
 		}
 		public StatSamples GetStat(string name)
@@ -398,7 +416,7 @@ namespace CSVStats
             }
             return null;
         }
-        public void AddStat(StatSamples stat)
+		public void AddStat(StatSamples stat)
         {           
             Stats.Add(stat.Name.ToLower(), stat);
         }
diff --git a/Engine/Source/Programs/CSVTools/PerfReportTool/PerfReportTool.cs b/Engine/Source/Programs/CSVTools/PerfReportTool/PerfReportTool.cs
--- a/Engine/Source/Programs/CSVTools/PerfReportTool/PerfReportTool.cs
+++ b/Engine/Source/Programs/CSVTools/PerfReportTool/PerfReportTool.cs
@@ -19,7 +19,7 @@ namespace PerfReportTool
 {
     class Version
     {
-        private static string VersionString = "3.97";
+        private static string VersionString = "4.00";
 
         public static string Get() { return VersionString; }
     };
@@ -262,6 +262,9 @@ namespace PerfReportTool
 			"       -writeSummaryCsv : if specified, a csv file containing summary information will be generated. Not available in bulk mode.\n" +
 			"       -beginEvent <event> : strip data before this event\n"+
 			"       -endEvent <event> : strip data after this event\n" +
+			"       -batchedGraphs : enable batched graph generation\n" +
+			"       -graphThreads : use with -batchedGraphs to control the number of threads per CsvToSVG instance (default: PC core count/2)\n" +
+			"       -nocommandlineEmbed : don't embed the commandline "+
 			"\n" +
 			"Optional bulk mode args: (use with -csvdir)\n" +
 			"       -recurse \n" +
@@ -293,7 +296,6 @@ namespace PerfReportTool
 			*/
 
 		Dictionary<string, string> statDisplaynameMapping;
-        Dictionary<string, Process> SvgProcesses = new Dictionary<string,Process>();
 		ReportXML reportXML;
 
 		string GetBaseDirectory()
@@ -322,6 +324,7 @@ namespace PerfReportTool
 				WriteLine(formatString);
 				return;
 			}
+			WriteLine("PerfReportTool v" + Version.Get());
 
 			ReadCommandLine(args);
             PerfLog perfLog = new PerfLog(GetBoolArg("perfLog"));
@@ -395,11 +398,17 @@ namespace PerfReportTool
                 }
             }
 
-            perfLog.LogTiming("Initialization");
+			int precacheCount = GetIntArg("precacheCount", 8);
+			int precacheThreads = GetIntArg("precacheThreads", 8);
+			bool bBatchedGraphs = GetBoolArg("batchedGraphs");
+			if (bBatchedGraphs)
+			{
+				WriteLine("Batched graph generation enabled.");
+			}
+
+			perfLog.LogTiming("Initialization");
 
-            int precacheCount = GetIntArg("precacheCount", 8);
-            int precacheThreads = GetIntArg("precacheThreads", 8);
-            CsvFileCache csvFileCache = new CsvFileCache(csvFilenames, precacheCount, precacheThreads);
+			CsvFileCache csvFileCache = new CsvFileCache(csvFilenames, precacheCount, precacheThreads);
             SummaryMetadataTable metadataTable = new SummaryMetadataTable();
             for ( int i=0; i<csvFilenames.Length; i++)
 			{
@@ -425,7 +434,7 @@ namespace PerfReportTool
 
                     if (bGenerateReport)
                     {
-                        GenerateReport(cachedCsvFile, outputDir, bBulkMode, metadata);
+                        GenerateReport(cachedCsvFile, outputDir, bBulkMode, metadata, bBatchedGraphs);
                         if (metadata != null)
                         {
                             metadataTable.Add(metadata);
@@ -617,7 +626,7 @@ namespace PerfReportTool
             bool loggingEnabled;
         }
 
-        void GenerateReport(CachedCsvFile csvFile, string outputDir, bool bBulkMode, SummaryMetadata summaryMetadata)
+        void GenerateReport(CachedCsvFile csvFile, string outputDir, bool bBulkMode, SummaryMetadata summaryMetadata, bool bBatchedGraphs)
         {
             PerfLog perfLog = new PerfLog(GetBoolArg("perfLog"));
             string shortName = ReplaceFileExtension(MakeShortFilename(csvFile.filename), "");
@@ -686,16 +695,23 @@ namespace PerfReportTool
             string [] statsToSummarise = statsToSummariseList.ToArray();
 
             float thickness = 1.0f;
-
-            List<string> svgFilenames = new List<string>();
+			List<string> csvToSvgCommandlines = new List<string>();
+			List<string> svgFilenames = new List<string>();
+			string responseFilename = null;
+			List<Process> csvToSvgProcesses = new List<Process>();
             if (writeDetailedReport)
             {
-                // Generate all the graphs asyncronously
+				// Generate all the graphs asyncronously
                 foreach (ReportGraph graph in reportTypeInfo.graphs)
                 {
-                    bool bFoundStat = false;
-
-                    foreach (string statString in graph.settings.statString.value.Split(' '))
+					string svgFilename = String.Empty;
+					if (graph.isExternal && !GetBoolArg("externalGraphs"))
+					{
+						svgFilenames.Add(svgFilename);
+						continue;
+					}
+					bool bFoundStat = false;
+					foreach (string statString in graph.settings.statString.value.Split(' '))
                     {
                         List<StatSamples> matchingStats = csvFile.dummyCsvStats.GetStatsMatchingString(statString);
                         if (matchingStats.Count > 0)
@@ -705,22 +721,33 @@ namespace PerfReportTool
                         }
 
                     }
-                    if (!bFoundStat)
-                    {
-						svgFilenames.Add(String.Empty);
-						continue;
-					}
-                    float graphScale = GetFloatArg("graphScale", 1.0f);
-
-					if (graph.isExternal && !GetBoolArg("externalGraphs"))
+					if (bFoundStat)
 					{
-						svgFilenames.Add(String.Empty);
-						continue;
+						float graphScale = GetFloatArg("graphScale", 1.0f);
+						svgFilename = GetTempFilename(csvFile.filename) + ".svg";
+						string args = GetCsvToSvgArgs(csvFile.filename, svgFilename, graph, thickness, minX, maxX, false, svgFilenames.Count);
+						if (bBatchedGraphs)
+						{
+							csvToSvgCommandlines.Add(args);
+						}
+						else
+						{
+							Process csvToSvgProcess = LaunchCsvToSvgAsync(args);
+							csvToSvgProcesses.Add(csvToSvgProcess);
+						}
 					}
-					string svgFilename = MakeGraphAsync(csvFile.filename, graph, thickness, minX, maxX, false, graphScale);
 					svgFilenames.Add(svgFilename);
                 }
-            }
+
+				if (bBatchedGraphs)
+				{
+					// Save the response file
+					responseFilename = GetTempFilename(csvFile.filename) + "_response.txt";
+					System.IO.File.WriteAllLines(responseFilename, csvToSvgCommandlines);
+					Process csvToSvgProcess = LaunchCsvToSvgAsync("-batchCommands "+responseFilename +" -mt " + GetIntArg("graphThreads", Environment.ProcessorCount/2).ToString() );
+					csvToSvgProcesses.Add(csvToSvgProcess);
+				}
+			}
             perfLog.LogTiming("    Initial Processing");
 
             // Read the csv stats while we wait for the graphs to complete
@@ -738,13 +765,10 @@ namespace PerfReportTool
 
             if ( writeDetailedReport )
             { 
-                // wait on the graphs to complete
-                foreach (string svgFilename in svgFilenames)
+                // wait on the graph processes to complete
+                foreach (Process process in csvToSvgProcesses)
 				{
-					if(!String.IsNullOrEmpty(svgFilename))
-					{
-						WaitForGraph(svgFilename);
-					}
+					process.WaitForExit();
 				}
                 perfLog.LogTiming("    WaitForAsyncGraphs");
             }
@@ -830,12 +854,16 @@ namespace PerfReportTool
             // Delete the temp files
             foreach (string svgFilename in svgFilenames)
             {
-				if(File.Exists(svgFilename))
+				if(svgFilename != String.Empty && File.Exists(svgFilename))
 				{
 					File.Delete(svgFilename);
 				}
-            }
-        }
+			}
+			if (responseFilename != null && File.Exists(responseFilename))
+			{
+				File.Delete(responseFilename);
+			}
+		}
 		CsvStats ReadCsvStats(CachedCsvFile csvFile, int minX, int maxX)
 		{
 			CsvStats csvStats = CsvStats.ReadCSVFromLines(csvFile.lines, null);
@@ -933,8 +961,15 @@ namespace PerfReportTool
 				htmlFile.WriteLine("<html>");
 				htmlFile.WriteLine("  <head>");
 				htmlFile.WriteLine("    <meta http-equiv='X-UA-Compatible' content='IE=edge'/>");
-				htmlFile.WriteLine("    <![CDATA[ \nCreated with PerfReportTool " + Version.Get() + " with commandline:");
-				htmlFile.WriteLine(commandLine);
+				if ( GetBoolArg("nocommandlineEmbed"))
+				{
+					htmlFile.WriteLine("    <![CDATA[ \nCreated with PerfReportTool " + Version.Get() );
+				}
+				else
+				{
+					htmlFile.WriteLine("    <![CDATA[ \nCreated with PerfReportTool " + Version.Get() + " with commandline:");
+					htmlFile.WriteLine(commandLine.GetCommandLine());
+				}
 				htmlFile.WriteLine("    ]]>");
 				htmlFile.WriteLine("    <title>" + titleStr + "</title>");
 				htmlFile.WriteLine("  </head>");
@@ -1097,7 +1132,7 @@ namespace PerfReportTool
             htmlFile.WriteLine("  <head>");
             htmlFile.WriteLine("    <meta http-equiv='X-UA-Compatible' content='IE=edge'/>");
             htmlFile.WriteLine("    <![CDATA[ \nCreated with PerfReportTool " + Version.Get() + " with commandline:");
-            htmlFile.WriteLine(commandLine);
+            htmlFile.WriteLine(commandLine.GetCommandLine());
             htmlFile.WriteLine("    ]]>");
             htmlFile.WriteLine("    <title>" + titleStr + "</title>");
             htmlFile.WriteLine("  </head>");
@@ -1145,25 +1180,29 @@ namespace PerfReportTool
             return str.Replace(" ", "");
         }
 
-        string MakeGraphAsync(string csvFilename, ReportGraph graph, double thicknessMultiplier, int minx, int maxx, bool multipleCSVs, float scaleby = 1.0f)
-        {
-            string title = graph.title;
-
-            GraphSettings graphSettings = graph.settings;
-            string statString = graphSettings.statString.value;
-            double thickness = graphSettings.thickness.value * thicknessMultiplier;
-            float maxy = GetFloatArg("maxy", (float)graphSettings.maxy.value);
-
-            bool smooth = graphSettings.smooth.value;
-            double smoothKernelPercent = graphSettings.smoothKernelPercent.value;
-            double smoothKernelSize = graphSettings.smoothKernelSize.value;
-            double compression = graphSettings.compression.value;
-            int width = graphSettings.width.value;
-            int height = graphSettings.height.value;
-            string additionalArgs = graphSettings.additionalArgs.value;
-            bool stacked = graphSettings.stacked.value;
-            bool showAverages = graphSettings.showAverages.value;
-            bool filterOutZeros = graphSettings.filterOutZeros.value;
+		string GetTempFilename(string csvFilename)
+		{
+			string shortFileName = MakeShortFilename(csvFilename).Replace(" ", "_");
+			return Path.Combine( Path.GetTempPath(), shortFileName+"_"+Guid.NewGuid().ToString().Substring(26));
+		}
+		string GetCsvToSvgArgs(string csvFilename, string svgFilename, ReportGraph graph, double thicknessMultiplier, int minx, int maxx, bool multipleCSVs, int graphIndex, float scaleby = 1.0f)
+		{
+			string title = graph.title;
+
+			GraphSettings graphSettings = graph.settings;
+			string statString = graphSettings.statString.value;
+			double thickness = graphSettings.thickness.value * thicknessMultiplier;
+			float maxy = GetFloatArg("maxy", (float)graphSettings.maxy.value);
+			bool smooth = graphSettings.smooth.value;
+			double smoothKernelPercent = graphSettings.smoothKernelPercent.value;
+			double smoothKernelSize = graphSettings.smoothKernelSize.value;
+			double compression = graphSettings.compression.value;
+			int width = graphSettings.width.value;
+			int height = graphSettings.height.value;
+			string additionalArgs = graphSettings.additionalArgs.value;
+			bool stacked = graphSettings.stacked.value;
+			bool showAverages = graphSettings.showAverages.value;
+			bool filterOutZeros = graphSettings.filterOutZeros.value;
 			bool snapToPeaks = false;
 			if (graphSettings.snapToPeaks.isSet)
 			{
@@ -1171,45 +1210,36 @@ namespace PerfReportTool
 			}
 
 			int lineDecimalPlaces = graphSettings.lineDecimalPlaces.isSet ? graphSettings.lineDecimalPlaces.value : 1;
-
 			int maxHierarchyDepth = graphSettings.maxHierarchyDepth.value;
-            string hideStatPrefix = graphSettings.hideStatPrefix.value;
-            string showEvents = graphSettings.showEvents.value;
-            double statMultiplier = graphSettings.statMultiplier.isSet ? graphSettings.statMultiplier.value : 1.0;
-
-            bool hideEventNames = false;
-            if ( multipleCSVs )
-            {
-                showEvents = "CSV:*";
-                hideEventNames = true;
-            }
-            bool interactive = true;
-
-            string outFilename = Path.GetTempFileName();
-
-
-            string svgName = title.Replace(" ", "_");
-            double budget = graph.budget;
-
-            string smoothParams = "";
-
-            if (smooth)
-            {
-                smoothParams = " -smooth";
-                if ( smoothKernelPercent >= 0.0f )
-                {
-                    smoothParams += " -smoothKernelPercent " + smoothKernelPercent.ToString();
-                }
-                if (smoothKernelSize >= 0.0f)
-                {
-                    smoothParams += " -smoothKernelSize " + smoothKernelSize.ToString();
-                }
-            }
+			string hideStatPrefix = graphSettings.hideStatPrefix.value;
+			string showEvents = graphSettings.showEvents.value;
+			double statMultiplier = graphSettings.statMultiplier.isSet ? graphSettings.statMultiplier.value : 1.0;
+			bool hideEventNames = false;
+			if (multipleCSVs)
+			{
+				showEvents = "CSV:*";
+				hideEventNames = true;
+			}
+			bool interactive = true;
+			double budget = graph.budget;
+			string smoothParams = "";
+			if (smooth)
+			{
+				smoothParams = " -smooth";
+				if (smoothKernelPercent >= 0.0f)
+				{
+					smoothParams += " -smoothKernelPercent " + smoothKernelPercent.ToString();
+				}
+				if (smoothKernelSize >= 0.0f)
+				{
+					smoothParams += " -smoothKernelSize " + smoothKernelSize.ToString();
+				}
+			}
 
-            string highlightEventRegions = "";
-            if (!GetBoolArg("noStripEvents"))
-            {
-                List<CsvEventStripInfo> eventsToStrip = reportXML.GetCsvEventsToStrip();
+			string highlightEventRegions = "";
+			if (!GetBoolArg("noStripEvents"))
+			{
+				List<CsvEventStripInfo> eventsToStrip = reportXML.GetCsvEventsToStrip();
 				if (eventsToStrip != null)
 				{
 					highlightEventRegions += "\"";
@@ -1222,57 +1252,57 @@ namespace PerfReportTool
 						string endEvent = (eventsToStrip[i].endName == null) ? "{NULL}" : eventsToStrip[i].endName;
 						highlightEventRegions += eventsToStrip[i].beginName + "," + endEvent;
 					}
-	                highlightEventRegions += "\"";
+					highlightEventRegions += "\"";
 				}
-            }
+			}
 
 			OptionalDouble minFilterStatValueSetting = graph.minFilterStatValue.isSet ? graph.minFilterStatValue : graphSettings.minFilterStatValue;
 
-			string shortFileName = MakeShortFilename(csvFilename).Replace(" ", "_");
-
-			outFilename += "_" + shortFileName + svgName.Replace("/", "_SLASH_").Replace("\\", "_BSLASH_") + ".svg";
-            string args =
-                " -csvs \""  + csvFilename  + "\"" +
-                " -title \"" + title + "\"" +
-                " -o "       + outFilename  +
-                " -stats "   + statString   +
-                " -width "   + (width*scaleby).ToString()  +
-                " -height "  + (height*scaleby).ToString() +
-                " -budget "  + budget.ToString() +
-                " -maxy "    + maxy.ToString()   +
-				" -lineDecimalPlaces "+ lineDecimalPlaces.ToString() +
-
-				( ( statMultiplier != 1.0)   ? " -statMultiplier " + statMultiplier.ToString("0.0000000000000000000000") : "") +
-                (hideEventNames              ? " -hideeventNames 1" : "") +
-                ( ( minx > 0 )               ? ( " -minx " + minx.ToString() ) : "" ) +
-                ( ( maxx != Int32.MaxValue ) ? ( " -maxx " + maxx.ToString() ) : "" ) +
+			string args =
+				" -csvs \"" + csvFilename + "\"" +
+				" -title \"" + title + "\"" +
+				" -o " + svgFilename +
+				" -stats " + statString +
+				" -width " + (width * scaleby).ToString() +
+				" -height " + (height * scaleby).ToString() +
+				" -budget " + budget.ToString() +
+				" -maxy " + maxy.ToString() +
+				" -uniqueID Graph_" + graphIndex.ToString() +
+				" -lineDecimalPlaces " + lineDecimalPlaces.ToString() +
+				" -nocommandlineEmbed "+
+
+				((statMultiplier != 1.0) ? " -statMultiplier " + statMultiplier.ToString("0.0000000000000000000000") : "") +
+				(hideEventNames ? " -hideeventNames 1" : "") +
+				((minx > 0) ? (" -minx " + minx.ToString()) : "") +
+				((maxx != Int32.MaxValue) ? (" -maxx " + maxx.ToString()) : "") +
 				OptionalHelper.GetDoubleSetting(graphSettings.miny, " -miny ") +
 				OptionalHelper.GetDoubleSetting(graphSettings.threshold, " -threshold ") +
 				OptionalHelper.GetDoubleSetting(graphSettings.averageThreshold, " -averageThreshold ") +
 				OptionalHelper.GetDoubleSetting(minFilterStatValueSetting, " -minFilterStatValue ") +
 				OptionalHelper.GetStringSetting(graphSettings.minFilterStatName, " -minFilterStatName ") +
-				(compression > 0.0           ? " -compression " + compression.ToString() : "") +
-                (thickness > 0.0             ? " -thickness " + thickness.ToString() : "") +
-                smoothParams + 
-                (interactive                  ? " -interactive" : "") +
-                (stacked                      ? " -stacked -forceLegendSort" : "") +
-                (showAverages                 ? " -showAverages" : "") +
-				(snapToPeaks                  ? "" : " -nosnap") +
-				(filterOutZeros               ? " -filterOutZeros" : "") +
-                (maxHierarchyDepth >=0        ? " -maxHierarchyDepth " + maxHierarchyDepth.ToString() : "") +
-                (hideStatPrefix.Length>0      ? " -hideStatPrefix " + hideStatPrefix : "") +
-                (graphSettings.mainStat.isSet ? " -stacktotalstat " + graphSettings.mainStat.value : "") +
-                (showEvents.Length>0          ? " -showevents "+showEvents : "") +
-                (highlightEventRegions.Length>0 ? " -highlightEventRegions " + highlightEventRegions : "") +
+				(compression > 0.0 ? " -compression " + compression.ToString() : "") +
+				(thickness > 0.0 ? " -thickness " + thickness.ToString() : "") +
+				smoothParams +
+				(interactive ? " -interactive" : "") +
+				(stacked ? " -stacked -forceLegendSort" : "") +
+				(showAverages ? " -showAverages" : "") +
+				(snapToPeaks ? "" : " -nosnap") +
+				(filterOutZeros ? " -filterOutZeros" : "") +
+				(maxHierarchyDepth >= 0 ? " -maxHierarchyDepth " + maxHierarchyDepth.ToString() : "") +
+				(hideStatPrefix.Length > 0 ? " -hideStatPrefix " + hideStatPrefix : "") +
+				(graphSettings.mainStat.isSet ? " -stacktotalstat " + graphSettings.mainStat.value : "") +
+				(showEvents.Length > 0 ? " -showevents " + showEvents : "") +
+				(highlightEventRegions.Length > 0 ? " -highlightEventRegions " + highlightEventRegions : "") +
 				(graphSettings.legendAverageThreshold.isSet ? " -legendAverageThreshold " + graphSettings.legendAverageThreshold.value : "") +
 
 				(graphSettings.ignoreStats.isSet ? " -ignoreStats " + graphSettings.ignoreStats.value : "") +
-                " "+additionalArgs;
-
-
-
+				" " + additionalArgs;
+			return args;
+		}
 
-            string csvToolPath = GetBaseDirectory() + "/CSVToSVG.exe";
+		Process LaunchCsvToSvgAsync(string args)
+		{
+			string csvToolPath = GetBaseDirectory() + "/CSVToSVG.exe";
 			string binary = csvToolPath;
 
 			// run mono on non-Windows hosts
@@ -1289,9 +1319,7 @@ namespace PerfReportTool
             startInfo.CreateNoWindow = true;
             startInfo.UseShellExecute = false;
             Process process = Process.Start(startInfo);
-
-            SvgProcesses[outFilename] = process;
-            return outFilename;
+			return process;
         }
 
         int CountCSVs( CsvStats csvStats)
@@ -1317,14 +1345,6 @@ namespace PerfReportTool
             return csvCount;
         }
 
-        void WaitForGraph(string filename)
-        {
-            Process process = SvgProcesses[filename];
-            process.WaitForExit();
-        }
-
-
-
         static int Main(string[] args)
         {
             Program program = new Program();
