PerfReportTool 4.01
- Batched graph generation is enabled by default (disable with -nobatchedgraphs)
- Report generation is 33% faster and no longer consumes the entire CPU

diff --git a/Engine/Source/Programs/CSVTools/PerfReportTool/PerfReportTool.cs b/Engine/Source/Programs/CSVTools/PerfReportTool/PerfReportTool.cs
--- a/Engine/Source/Programs/CSVTools/PerfReportTool/PerfReportTool.cs
+++ b/Engine/Source/Programs/CSVTools/PerfReportTool/PerfReportTool.cs
@@ -19,7 +19,7 @@ namespace PerfReportTool
 {
     class Version
     {
-        private static string VersionString = "4.00";
+        private static string VersionString = "4.01";
 
         public static string Get() { return VersionString; }
     };
@@ -262,7 +262,7 @@ namespace PerfReportTool
 			"       -writeSummaryCsv : if specified, a csv file containing summary information will be generated. Not available in bulk mode.\n" +
 			"       -beginEvent <event> : strip data before this event\n"+
 			"       -endEvent <event> : strip data after this event\n" +
-			"       -batchedGraphs : enable batched graph generation\n" +
+			"       -noBatchedGraphs : disable batched/multithreaded graph generation (default is enabled)\n" +
 			"       -graphThreads : use with -batchedGraphs to control the number of threads per CsvToSVG instance (default: PC core count/2)\n" +
 			"       -nocommandlineEmbed : don't embed the commandline "+
 			"\n" +
@@ -400,11 +400,19 @@ namespace PerfReportTool
 
 			int precacheCount = GetIntArg("precacheCount", 8);
 			int precacheThreads = GetIntArg("precacheThreads", 8);
-			bool bBatchedGraphs = GetBoolArg("batchedGraphs");
+			bool bBatchedGraphs = true;
+			if ( GetBoolArg("noBatchedGraphs") )
+			{
+				bBatchedGraphs = false;
+			}
 			if (bBatchedGraphs)
 			{
 				WriteLine("Batched graph generation enabled.");
 			}
+			else
+			{
+				WriteLine("Batched graph generation disabled.");
+			}
 
 			perfLog.LogTiming("Initialization");
 
