CsvCollate 1.31
- Fixes an out-by-one bug with -avg. The stats were being divided by csvCount-1 instead of csvCount

diff --git a/Engine/Source/Programs/CSVTools/CSVCollate/CSVCollate.cs b/Engine/Source/Programs/CSVTools/CSVCollate/CSVCollate.cs
--- a/Engine/Source/Programs/CSVTools/CSVCollate/CSVCollate.cs
+++ b/Engine/Source/Programs/CSVTools/CSVCollate/CSVCollate.cs
@@ -14,7 +14,7 @@ namespace CSVTools
 {
     class Version
     {
-        private static string VersionString = "1.30";
+        private static string VersionString = "1.31";
 
         public static string Get() { return VersionString; }
     };
@@ -68,13 +68,37 @@ namespace CSVTools
 			// Whether or not we want stats to be averaged rather than appended
 			bool bAverage = GetBoolArg("avg");
 
-            // Read CSV filenames from a directory or list
-            string[] csvFilenames;
-            string csvDir = GetArg("csvDir");
+			string csvDir = GetArg("csvDir");
+			string csvFilenamesStr = GetArg("csvs", false);
+
+			if ( csvDir == null && csvFilenamesStr == null )
+			{
+				WriteLine("-csvs or -csvdir is required");
+				WriteLine(formatString);
+				return;
+			}
+
+			string searchPattern = GetArg("searchPattern", null);
+			if (csvFilenamesStr.Contains("*"))
+			{
+				// If passed a wildcard to -csvs, this is equivalent to -csvdir . -searchpattern <csvs>
+				if (csvDir != null && csvDir != "")
+				{
+					throw new Exception("Can't use -csvs with -csvdir");
+				}
+				csvDir = ".";
+				searchPattern = csvFilenamesStr;
+			}
+
+			// Read CSV filenames from a directory or list
+			string[] csvFilenames;
             if (csvDir.Length > 0)
             {
-                DirectoryInfo di = new DirectoryInfo(csvDir);
-				string searchPattern = GetArg("searchPattern","*.csv");
+				if (searchPattern == null)
+				{
+					searchPattern = "*.csv";
+				}
+				DirectoryInfo di = new DirectoryInfo(csvDir);
 				bool bRecurse = GetBoolArg("recurse");
 				var files = di.GetFiles(searchPattern, bRecurse ? SearchOption.AllDirectories : SearchOption.TopDirectoryOnly);
                 csvFilenames = new string[files.Length];
@@ -87,7 +111,6 @@ namespace CSVTools
             }
             else
             {
-                string csvFilenamesStr = GetArg("csvs", true);
                 if (csvFilenamesStr.Length == 0)
                 {
                     System.Console.Write(formatString);
@@ -161,24 +184,28 @@ namespace CSVTools
 						metadataA.CombineAndValidate(metadataB);
 					}
 					combinedCsvStats.Combine(srcCsvStats, bAverage, false);
-					if ( bAverage )
+				}
+
+				// If we're computing the average, update the counts for each frame
+				if (bAverage)
+				{
+					// Resize frameCsvCounts if necessary
+					for (int i = frameCsvCounts.Count; i < combinedCsvStats.SampleCount; i++)
 					{
-						// Resize the framecount array to match
-						for (int i= frameCsvCounts.Count; i<combinedCsvStats.SampleCount; i++)
-						{
-							frameCsvCounts.Add(0);
-						}
-						for (int i=0; i<srcCsvStats.SampleCount;i++)
-						{
-							frameCsvCounts[i] += 1;
-						}
+						frameCsvCounts.Add(0);
+					}
+					for (int i = 0; i < srcCsvStats.SampleCount; i++)
+					{
+						frameCsvCounts[i] += 1;
 					}
 				}
+
 				allCsvFilenames.Add(Path.GetFileName(csvFilename));
 				csvIndex++;
 				WriteLine("Csvs Processed: " + csvIndex + " / " + csvFilenames.Length);
 			}
 
+
 			if (bAverage)
 			{
 				// Divide all samples by the total number of CSVs 
