CsvToSVG 3.31
- Reduce event spam by grouping duplicate events if they appear within a certain number of pixels of each other. Display the count
- Also reduce the frequency of event lines and reduce the alpha when spamming occurs

diff --git a/Engine/Source/Programs/CSVTools/CSVToSVG/CsvToSVG.cs b/Engine/Source/Programs/CSVTools/CSVToSVG/CsvToSVG.cs
--- a/Engine/Source/Programs/CSVTools/CSVToSVG/CsvToSVG.cs
+++ b/Engine/Source/Programs/CSVTools/CSVToSVG/CsvToSVG.cs
@@ -25,7 +25,7 @@ namespace CSVTools
 {
     class Version
     {
-        private static string VersionString = "2.3";
+        private static string VersionString = "2.31";
         
         public static string Get() { return VersionString; }
     };
@@ -1805,47 +1805,106 @@ namespace CSVTools
         }
 
 
+		class CsvEventWithCount : CsvEvent
+		{
+			public CsvEventWithCount(CsvEvent ev)
+			{
+				base.Frame = ev.Frame;
+				base.Name = ev.Name;
+				count = 1;
+			}
+			public int count;
+		};
 
-        void DrawEventText(List<CsvEvent> events, Colour colour, Rect rect, Range range )
+		void DrawEventText(List<CsvEvent> events, Colour colour, Rect rect, Range range )
         {
-            float LastEventX = -100000.0f;
-            int lastFrame = 0;
-            foreach (CsvEvent ev in events)
-            {
-                // Only draw events which are in the range
-                if (ev.Frame >= range.MinX && ev.Frame <= range.MaxX)
-                {
-                    float eventX = ToSvgX(ev.Frame, rect, range);
-                    float DistToLastEvent = Math.Abs(eventX - LastEventX);
+			float LastEventX = -100000.0f;
+			int lastFrame = 0;
+
+			// Make a filtered list of events to display, grouping duplicates which are close together
+			List<CsvEventWithCount> filteredEvents = new List<CsvEventWithCount>();
+			CsvEventWithCount currentDisplayEvent = null;
+			CsvEvent lastEvent = null;
+			foreach (CsvEvent ev in events)
+			{
+				// Only draw events which are in the range
+				if (ev.Frame >= range.MinX && ev.Frame <= range.MaxX)
+				{
+					// Merge with the current display event?
+					bool bMerge = false;
+					if (currentDisplayEvent != null && lastEvent != null && ev.Name == currentDisplayEvent.Name)
+					{
+						float DistToLastEvent = ToSvgX(ev.Frame, rect, range) - ToSvgX(lastEvent.Frame, rect, range);
+						if (DistToLastEvent <= 8.5)
+						{
+							bMerge = true;
+						}
+					}
 
-                    string name = ev.Name;
+					if (bMerge)
+					{
+						currentDisplayEvent.count++;
+					}
+					else
+					{
+						currentDisplayEvent = new CsvEventWithCount(ev);
+						filteredEvents.Add(currentDisplayEvent);
+					}
+					lastEvent = ev;
+				}
+			}
 
-                    // Space out the events (allow at least 8 pixels between them)
-                    if (eventX - LastEventX <= 8.5f)
-                    {
-                        // Add an arrow to indicate events were spaced out 
-						name = "&#x21b7; " + name + " (+" + (ev.Frame - lastFrame) + ")";
-						eventX = LastEventX + 9.0f;
-                    }
-                    float csvTextX = ToCsvX(eventX + 7, rect, range);
+			foreach (CsvEventWithCount ev in filteredEvents)
+            {
+				float eventX = ToSvgX(ev.Frame, rect, range);
+                string name = ev.Name;
 
-                    DrawHorizontalAxisText(name, csvTextX, colour, rect, range, 10, true);
-                    LastEventX = eventX;
-                    lastFrame = ev.Frame;
+				if ( ev.count > 1 )
+				{
+					name += " &#x00D7; " + ev.count;
+				}
+
+                // Space out the events (allow at least 8 pixels between them)
+                if (eventX - LastEventX <= 8.5f)
+                {
+					// Add an arrow to indicate events were spaced out 
+					name = "&#x21b7; " + name;
+					if ( ev.count == 1 )
+					{
+						name += " (+" + (ev.Frame - lastFrame) + ")";
+					}
+					eventX = LastEventX + 9.0f;
                 }
+                float csvTextX = ToCsvX(eventX + 7, rect, range);
+
+                DrawHorizontalAxisText(name, csvTextX, colour, rect, range, 10, true);
+                LastEventX = eventX;
+                lastFrame = ev.Frame;
             }
         }
 
         void DrawEventLines(List<CsvEvent> events, Rect rect, Range range)
         {
-            foreach (CsvEvent ev in events)
+			float LastEventX = -100000.0f;
+			float LastDisplayedEventX = -100000.0f;
+			int i = 0;
+			foreach (CsvEvent ev in events)
             {
                 // Only draw events which are in the range
                 if (ev.Frame >= range.MinX && ev.Frame <= range.MaxX)
                 {
-                    DrawVerticalLine(ev.Frame, theme.MajorGridlineColour, rect, range, 1.0f, true, true);
-                }
-            }
+					float eventX = ToSvgX(ev.Frame, rect, range);
+					// Space out the events (allow at least 2 pixels between them)
+					if (eventX - LastDisplayedEventX > 3)
+					{
+						float alpha = (eventX - LastEventX < 1.0f) ? 0.5f : 1.0f;
+						DrawVerticalLine(ev.Frame, theme.MajorGridlineColour, rect, range, alpha, true, true);
+						LastDisplayedEventX = eventX;
+					}
+					LastEventX = eventX;
+				}
+				i++;
+			}
         }
 
 
