CsvCollate 1.32- metadata filtering support:
-metadatafilter "key0=value0,key1=value1,..."

diff --git a/Engine/Source/Programs/CSVTools/CSVCollate/CSVCollate.cs b/Engine/Source/Programs/CSVTools/CSVCollate/CSVCollate.cs
--- a/Engine/Source/Programs/CSVTools/CSVCollate/CSVCollate.cs
+++ b/Engine/Source/Programs/CSVTools/CSVCollate/CSVCollate.cs
@@ -14,7 +14,7 @@ namespace CSVTools
 {
     class Version
     {
-        private static string VersionString = "1.31";
+        private static string VersionString = "1.32";
 
         public static string Get() { return VersionString; }
     };
@@ -29,9 +29,9 @@ namespace CSVTools
 			"       [-recurse] - for use with -csvdir\n" +
 			"       [-filterOutlierStat <stat>] - discard CSVs if this stat has very high values\n" +
 			"       [-filterOutlierThreshold <value>] - threshold for outliers (default:1000)\n" +
+			"       [-metadataFilter <key=value,key=value...>] : filters based on CSV metadata\n" +
 			"       -o <csvFilename> \n";
 
-
 		void Run(string[] args)
         {
             // Read the command line
@@ -129,6 +129,7 @@ namespace CSVTools
 			CsvStats combinedCsvStats = new CsvStats();
 
 
+			string metadataFilterString = GetArg("metadataFilter", null);
 			List<int> frameCsvCounts=new List<int>();
 			List<string> allCsvFilenames = new List<string>();
 			int csvIndex = 0;
@@ -154,6 +155,14 @@ namespace CSVTools
 						}
 					}
 				}
+				if (metadataFilterString != null)
+				{
+					if (srcCsvStats.metaData == null || !CsvStats.DoesMetadataMatchFilter(srcCsvStats.metaData, metadataFilterString))
+					{
+						WriteLine("Skipping CSV " + csvFilename + " due to metadata filter");
+						skip = true;
+					}
+				}
 				if ( skip )
 				{
 					continue;
diff --git a/Engine/Source/Programs/CSVTools/CsvStats/CsvStats.cs b/Engine/Source/Programs/CSVTools/CsvStats/CsvStats.cs
--- a/Engine/Source/Programs/CSVTools/CsvStats/CsvStats.cs
+++ b/Engine/Source/Programs/CSVTools/CsvStats/CsvStats.cs
@@ -664,8 +664,8 @@ namespace CSVStats
             // Write the metadata
         }
 
-        // Pad the stats
-        public void ComputeAveragesAndTotal()
+		// Pad the stats
+		public void ComputeAveragesAndTotal()
         {
             foreach (StatSamples stat in Stats.Values.ToArray())
             {
@@ -1019,7 +1019,31 @@ namespace CSVStats
             }
         }
 
-        public static CsvStats ReadCSVFile(string csvFilename, string[] statNames, int numRowsToSkip = 0)
+		public static bool DoesMetadataMatchFilter(CsvMetadata metadata, string metadataFilterString)
+		{
+			string[] keyValuePairStrs = metadataFilterString.Split(',');
+			foreach (string keyValuePairStr in keyValuePairStrs)
+			{
+				string[] keyValue = keyValuePairStr.Split('=');
+				if (keyValue.Length != 2)
+				{
+					return false;
+				}
+				string key = keyValue[0].ToLower();
+				if (!metadata.Values.ContainsKey(key))
+				{
+					return false;
+				}
+				if (metadata.Values[key].ToLower() != keyValue[1].ToLower())
+				{
+					return false;
+				}
+			}
+			return true;
+		}
+
+
+		public static CsvStats ReadCSVFile(string csvFilename, string[] statNames, int numRowsToSkip = 0)
         {
             string [] lines = ReadLinesFromFile(csvFilename);
             return ReadCSVFromLines(lines, statNames, numRowsToSkip);
diff --git a/Engine/Source/Programs/CSVTools/PerfReportTool/PerfReportTool.cs b/Engine/Source/Programs/CSVTools/PerfReportTool/PerfReportTool.cs
--- a/Engine/Source/Programs/CSVTools/PerfReportTool/PerfReportTool.cs
+++ b/Engine/Source/Programs/CSVTools/PerfReportTool/PerfReportTool.cs
@@ -19,7 +19,7 @@ namespace PerfReportTool
 {
     class Version
     {
-        private static string VersionString = "4.01";
+        private static string VersionString = "4.02";
 
         public static string Get() { return VersionString; }
     };
@@ -385,9 +385,6 @@ namespace PerfReportTool
 			reportXML = new ReportXML(GetArg("graphxml", false), GetArg("reportxml", false), GetArg("reportxmlbasedir", false));
 			statDisplaynameMapping = reportXML.GetDisplayNameMapping();
 
-            string metadataFilterString = "";
-            metadataFilterString = GetArg("metadataFilter");
-
             // Create the output directory if requested
             string outputDir = GetArg("o", false).ToLower();
             if (!string.IsNullOrEmpty(outputDir))
@@ -416,6 +413,8 @@ namespace PerfReportTool
 
 			perfLog.LogTiming("Initialization");
 
+			string metadataFilterString = GetArg("metadataFilter", "");
+
 			CsvFileCache csvFileCache = new CsvFileCache(csvFilenames, precacheCount, precacheThreads);
             SummaryMetadataTable metadataTable = new SummaryMetadataTable();
             for ( int i=0; i<csvFilenames.Length; i++)
@@ -575,25 +574,7 @@ namespace PerfReportTool
                 Console.WriteLine("CSV " + csvFile.filename + " has no metadata");
                 return false;
             }
-            string[] keyValuePairStrs = metadataFilterString.Split(',');
-            foreach (string keyValuePairStr in keyValuePairStrs)
-            {
-                string[] keyValue = keyValuePairStr.Split('=');
-                if (keyValue.Length != 2)
-                {
-                    return false;
-                }
-                string key = keyValue[0].ToLower();
-                if (!metadata.Values.ContainsKey(key))
-                {
-                    return false;
-                }
-                if (metadata.Values[key].ToLower() != keyValue[1].ToLower())
-                {
-                    return false;
-                }
-            }
-            return true;
+			return CsvStats.DoesMetadataMatchFilter(metadata, metadataFilterString);
         }
         class PerfLog
         {
