PerfReportTool - remove some dead code

diff --git a/Engine/Source/Programs/CSVTools/PerfReportTool/perfsummaries.cs b/Engine/Source/Programs/CSVTools/PerfReportTool/perfsummaries.cs
--- a/Engine/Source/Programs/CSVTools/PerfReportTool/perfsummaries.cs
+++ b/Engine/Source/Programs/CSVTools/PerfReportTool/perfsummaries.cs
@@ -883,11 +883,6 @@ namespace PerfSummaries
             }
             htmlFile.WriteLine("  </table>");
 			htmlFile.WriteLine("<p style='font-size:8'>Note: Simplified hitch metric. All frames over threshold are counted" + "</p>");
-
-			if (csvStats.Events.Count == 0)
-            {
-                return;
-            }
         }
         public double[] HitchThresholds;
     };
