PerfReportTool 4.03
  - support multiple local player views.

diff --git a/Engine/Source/Programs/CSVTools/PerfReportTool/PerfReportTool.cs b/Engine/Source/Programs/CSVTools/PerfReportTool/PerfReportTool.cs
--- a/Engine/Source/Programs/CSVTools/PerfReportTool/PerfReportTool.cs
+++ b/Engine/Source/Programs/CSVTools/PerfReportTool/PerfReportTool.cs
@@ -19,7 +19,7 @@ namespace PerfReportTool
 {
     class Version
     {
-        private static string VersionString = "4.02";
+        private static string VersionString = "4.03";
 
         public static string Get() { return VersionString; }
     };
diff --git a/Engine/Source/Programs/CSVTools/PerfReportTool/perfsummaries.cs b/Engine/Source/Programs/CSVTools/PerfReportTool/perfsummaries.cs
--- a/Engine/Source/Programs/CSVTools/PerfReportTool/perfsummaries.cs
+++ b/Engine/Source/Programs/CSVTools/PerfReportTool/perfsummaries.cs
@@ -1706,7 +1706,7 @@ namespace PerfSummaries
 
 					if (xStat == null || yStat == null)
 					{
-						throw new Exception("MapOverlaySummary: CSV doesn't contain stat " + overlay.positionStatNames[0] + " and/or " + overlay.positionStatNames[1]);
+						continue;
 					}
 
 					// If a startevent is specified, update the start frame
