PerfReportTool 4.32
 - Add -metadataCacheIn mode for running without CSVs
 - Use a BlockingCollection to eliminate precache threading bottleneck and simplify precache thread producer/consumer code
 - Threading still hits a performance bottleneck at around 6 threads when running metadata queries currently, and gets slower after that as more threads are added.
    - Need further work to optimize GetReportTypeXML (too much unnecessary linq) and SummaryMetadata::TryReadFromCache

diff --git a/Engine/Source/Programs/CSVTools/PerfReportTool/PerfReportTool.cs b/Engine/Source/Programs/CSVTools/PerfReportTool/PerfReportTool.cs
--- a/Engine/Source/Programs/CSVTools/PerfReportTool/PerfReportTool.cs
+++ b/Engine/Source/Programs/CSVTools/PerfReportTool/PerfReportTool.cs
@@ -12,6 +12,7 @@ using System.IO;
 using System.Diagnostics;
 using CSVStats;
 using System.Collections;
+using System.Collections.Concurrent;
 using System.Threading;
 using System.Security.Cryptography;
 
@@ -21,7 +22,7 @@ namespace PerfReportTool
 {
     class Version
     {
-        private static string VersionString = "4.31";
+        private static string VersionString = "4.32";
 
         public static string Get() { return VersionString; }
     };
@@ -243,15 +244,25 @@ namespace PerfReportTool
 
     class CachedCsvFile
     {
+		enum FileType
+		{
+			TextCsv,
+			BinaryCsv,
+			MetadataCachePrc
+		};
 		public CachedCsvFile(string inFilename, bool useCacheFiles, DerivedMetadataMappings derivedMetadataMappings)
 		{
 			if (inFilename.ToLower().EndsWith(".csv"))
 			{
-				isBinaryFile = false;
+				fileType = FileType.TextCsv;
 			}
 			else if (inFilename.ToLower().EndsWith(".csv.bin"))
 			{
-				isBinaryFile = true;
+				fileType = FileType.BinaryCsv;
+			}
+			else if (inFilename.ToLower().EndsWith(".prc"))
+			{
+				fileType = FileType.MetadataCachePrc;
 			}
 			else
 			{
@@ -275,11 +286,11 @@ namespace PerfReportTool
 			}
 			else
 			{
-				if (isBinaryFile)
+				if (fileType == FileType.BinaryCsv)
 				{
 					dummyCsvStats = CsvStats.ReadBinFile(inFilename, null, 0, true);
 				}
-				else
+				else if (fileType == FileType.TextCsv)
 				{
 					textCsvLines = File.ReadAllLines(inFilename);
 					if (textCsvLines.Length > 0)
@@ -292,11 +303,33 @@ namespace PerfReportTool
 						dummyCsvStats = new CsvStats();
 					}
 				}
-			} 
-
-            filename = inFilename;
+				else if (fileType == FileType.MetadataCachePrc)
+				{
+					// Reconstruct the dummy CSV from the metadata
+					dummyCsvStats = new CsvStats();
+					cachedSummaryMetadata = SummaryMetadata.TryReadFromCacheFile(inFilename);
+					foreach (KeyValuePair<string, SummaryMetadataValue> entry in cachedSummaryMetadata.dict)
+					{
+						if ( entry.Value.type == SummaryMetadataValue.Type.CsvMetadata )
+						{
+							if (dummyCsvStats.metaData == null)
+							{
+								dummyCsvStats.metaData = new CsvMetadata();
+							}
+							dummyCsvStats.metaData.Values.Add(entry.Key, entry.Value.value);
+						}
+						if ( entry.Value.type == SummaryMetadataValue.Type.CsvStatAverage )
+						{
+							dummyCsvStats.AddStat(new StatSamples(entry.Key));
+						}
+					}
+					metadata = dummyCsvStats.metaData;
+				}
+			}
+			filename = inFilename;
 
-            if (dummyCsvStats != null && dummyCsvStats.metaData != null)
+			// Apply metadata mappings if we have metadata. Note: we skip for PRCs, since these already have mappings cached
+            if (dummyCsvStats != null && dummyCsvStats.metaData != null && fileType != FileType.MetadataCachePrc)
             {
 				metadata = dummyCsvStats.metaData;
 				derivedMetadataMappings.ApplyMapping(metadata);
@@ -305,7 +338,7 @@ namespace PerfReportTool
 
 		public void PrepareCsvData()
 		{
-			if (isBinaryFile)
+			if (fileType == FileType.BinaryCsv)
 			{
 				finalCsv = CsvStats.ReadBinFile(filename);
 			}
@@ -324,7 +357,7 @@ namespace PerfReportTool
 			{
 				return finalCsv;
 			}
-			if (!isBinaryFile)
+			if (fileType == FileType.TextCsv)
 			{
 				finalCsv = CsvStats.ReadCSVFromLines(textCsvLines, null);
 				return finalCsv;
@@ -382,8 +415,8 @@ namespace PerfReportTool
 		public CsvMetadata metadata;
 		public SummaryMetadata cachedSummaryMetadata;
 		public ReportTypeInfo reportTypeInfo;
-		bool isBinaryFile;
-    };
+		FileType fileType;
+	};
 
 	class MetadataCacheStats
 	{
@@ -416,7 +449,7 @@ namespace PerfReportTool
 			"PerfReportTool v" + Version.Get() + "\n" +
 			"\n" +
 			"Format: \n" +
-			"       -csv <filename> or -csvdir <directory path>\n" +
+			"       -csv <filename> or -csvdir <directory path> or -metadataCacheIn <directory path>\n" +
 			"       -o <dir name>: output directory (will be created if necessary)\n" +
 			"\n" +
 			"Optional Args:\n" +
@@ -444,7 +477,7 @@ namespace PerfReportTool
 			"       -noStripEvents : if specified, don't strip out samples between excluded events from the stats\n" +
 			"NOTE: these options disable metadata caching\n" +
 			"\n" +
-			"Optional bulk mode args: (use with -csvdir)\n" +
+			"Optional bulk mode args: (use with -csvdir or -metadataCacheIn)\n" +
 			"       -recurse \n" +
 			"       -searchpattern <pattern>, e.g -searchpattern csvprofile*\n" +
 			"       -customTable <comma seprated fields>\n" +
@@ -474,6 +507,7 @@ namespace PerfReportTool
 			"Performance args for bulk mode:\n" +
 			"       -precacheCount <n> : number of CSV files to precache in the lookahead cache (0 for no precache)\n" +
 			"       -precacheThreads <n> : number of threads to use for the CSV lookahead cache (default 8)\n" +
+			"       -metadataCacheIn <dir> : reads data directly from the metadata cache instead of from CSVs\n" +
 			"       -metadataCache <dir> : specifies a directory for metadata to be cached. Enables -readAllStats implicitly.\n  this avoids processing csvs on subsequent runs if -noDetailedReports is specified\n" +
 			"       -metadataCacheInvalidate : regenerates metadata disk cache entries (ie write only)\n" +
 			"       -metadataCacheReadOnly : only read from the cache, never write\n" +
@@ -525,6 +559,7 @@ namespace PerfReportTool
             PerfLog perfLog = new PerfLog(GetBoolArg("perfLog"));
 
             bool bBulkMode = false;
+			bool bMetadataCacheOnlyMode = false;
 			// Read CSV filenames from a directory or list
 			string[] csvFilenames;
 			if (args.Length == 1)
@@ -535,6 +570,7 @@ namespace PerfReportTool
 			else
 			{
 				string csvDir = GetArg("csvDir");
+				string metadataCacheInDir = GetArg("metadataCacheIn");
 				if (csvDir.Length > 0)
 				{
 					bool recurse = GetBoolArg("recurse");
@@ -560,7 +596,22 @@ namespace PerfReportTool
 					bBulkMode = true;
                     perfLog.LogTiming("DirectoryScan");
                 }
-                else
+				else if (metadataCacheInDir.Length>0)
+				{
+					bool recurse = GetBoolArg("recurse");
+					System.IO.FileInfo[] files = GetFilesWithSearchPattern(metadataCacheInDir, "*.prc", recurse);
+					csvFilenames = new string[files.Length];
+					int i = 0;
+					foreach (FileInfo csvFile in files)
+					{
+						csvFilenames[i] = csvFile.FullName;
+						i++;
+					}
+					bBulkMode = true;
+					bMetadataCacheOnlyMode = true;
+					perfLog.LogTiming("DirectoryScan");
+				}
+				else
 				{
 					string csvFilenamesStr = GetArg("csv");
 					if (csvFilenamesStr.Length == 0)
@@ -673,16 +724,22 @@ namespace PerfReportTool
 			// Make PerfSummaryCache from: CSV hash + reporttype hash. If the cache is enabled then always -readAllStats
 
 			bool bReadAllStats = GetBoolArg("readAllStats");
-			bool bShowHiddenStats= GetBoolArg("showHiddenStats");
+
 			bool bMetadataCacheReadonly = GetBoolArg("metadataCacheReadOnly");
 			bool bMetadataCacheInvalidate = GetBoolArg("metadataCacheInvalidate");
-
 			string cleanCsvOutputFilename = GetArg("cleanCsvOut", null);
 			if (cleanCsvOutputFilename != null && bBulkMode)
 			{
 				throw new Exception("-cleanCsvOut is not compatible with bulk mode. Pass one csv with -csv <filename>");
 			}
 
+			bool bShowHiddenStats = GetBoolArg("showHiddenStats");
+			string customSummaryTableFilter = GetArg("customTable");
+			if (customSummaryTableFilter.Length > 0)
+			{
+				bShowHiddenStats = true;
+			}
+
 			string summaryMetadataCacheDirForRead = summaryMetadataCacheDir;
 			if ( bMetadataCacheInvalidate || writeDetailedReports )
 			{
@@ -690,6 +747,14 @@ namespace PerfReportTool
 				summaryMetadataCacheDirForRead = null;
 			}
 
+			if (bMetadataCacheOnlyMode)
+			{
+				// Override these options in metadatacacheonly mode
+				bMetadataCacheReadonly = true;
+				summaryMetadataCacheDirForRead = null;
+				bMetadataCacheInvalidate = false;
+			}
+
 			ReportTypeParams reportTypeParams = new ReportTypeParams
 			{
 				reportTypeOverride = GetArg("reportType", false).ToLower(),
@@ -697,18 +762,23 @@ namespace PerfReportTool
 			};
 
 			bool bRemoveDuplicates = !GetBoolArg("allowDuplicateCSVs");
-			CsvFileCache csvFileCache = new CsvFileCache(csvFilenames, precacheCount, precacheThreads, !GetBoolArg("noCsvCacheFiles"), metadataQuery, reportXML, reportTypeParams, bBulkMode, bRemoveDuplicates, summaryMetadataCacheDirForRead);
+			CsvFileCache csvFileCache = new CsvFileCache(csvFilenames, precacheCount, precacheThreads, !GetBoolArg("noCsvCacheFiles"), metadataQuery, reportXML, reportTypeParams, bBulkMode, bMetadataCacheOnlyMode, bRemoveDuplicates, summaryMetadataCacheDirForRead);
             SummaryMetadataTable metadataTable = new SummaryMetadataTable();
 			bool bWriteToMetadataCache = summaryMetadataCacheDir != null && !bMetadataCacheReadonly;
 
-			for ( int i=0; i<csvFilenames.Length; i++)
+			int csvCount = csvFilenames.Length;
+			for ( int i=0; i<csvCount; i++)
 			{
                 try
                 {
                     CachedCsvFile cachedCsvFile = csvFileCache.GetNextCachedCsvFile();
+					if (cachedCsvFile == null)
+					{
+						continue;
+					}
                     Console.WriteLine("-------------------------------------------------");
                     Console.WriteLine("CSV " + (i+1) + "/" + csvFilenames.Length ) ;
-                    Console.WriteLine(csvFilenames[i] );
+                    Console.WriteLine(cachedCsvFile.filename);
 
 					perfLog.LogTiming("  CsvCacheRead");
 					if (cachedCsvFile == null)
@@ -731,6 +801,7 @@ namespace PerfReportTool
 							if (cleanCsvOutputFilename != null)
 							{
 								WriteCleanCsv(cachedCsvFile, cleanCsvOutputFilename, cachedCsvFile.reportTypeInfo);
+								perfLog.LogTiming("  WriteCleanCsv");
 							}
 							else
 							{
@@ -740,7 +811,7 @@ namespace PerfReportTool
 								{
 									if (metadata.WriteToCache(summaryMetadataCacheDir, cachedCsvFile.summaryTableCacheId))
 									{
-										Console.WriteLine("Cached summary metadata for CSV: " + csvFilenames[i]);
+										Console.WriteLine("Cached summary metadata for CSV: " + cachedCsvFile.filename);
 										metadataCacheStats.WriteCount++;
 										perfLog.LogTiming("  WriteMetadataCache");
 									}
@@ -782,13 +853,11 @@ namespace PerfReportTool
 				{
 					summaryTableFilename = summaryTableFilename.Substring(0, summaryTableFilename.Length - 5);
 				}
-				string customSummaryTableFilter = GetArg("customTable");
 				bool bCsvTable = GetBoolArg("csvTable");
 				bool bCollateTable = GetBoolArg("collateTable");
 				bool bSpreadsheetFriendlyStrings = GetBoolArg("spreadsheetFriendly");
 				if (customSummaryTableFilter.Length > 0)
 				{
-					bShowHiddenStats = true;
 					string customSummaryTableRowSort = GetArg("customTableSort");
 					if (customSummaryTableRowSort.Length == 0)
 					{
@@ -908,7 +977,7 @@ namespace PerfReportTool
                 if (loggingEnabled)
                 {
                     Console.WriteLine("[PerfLog] "+
-                        String.Format("{0,-25} : {1,-10}", description , elapsed.ToString("0.00") + "s"), 70);
+                        String.Format("{0,-25} : {1,-10}", description , (elapsed*1000.0).ToString("0.0") + "ms"), 70);
                     if ( newLine )
                     {
                         Console.WriteLine();
@@ -2407,7 +2476,7 @@ namespace PerfReportTool
 
 	class CsvFileCache
     {
-        public CsvFileCache( string[] inCsvFilenames, int inLookaheadCount, int inThreadCount, bool inUseCacheFiles, QueryExpression inMetadataQuery, ReportXML inReportXml, ReportTypeParams inReportTypeParams, bool inBulkMode, bool inRemoveDuplicates, string inSummaryMetadataCacheDir = null )
+        public CsvFileCache( string[] inCsvFilenames, int inLookaheadCount, int inThreadCount, bool inUseCacheFiles, QueryExpression inMetadataQuery, ReportXML inReportXml, ReportTypeParams inReportTypeParams, bool inBulkMode, bool inMetadataCacheOnlyMode, bool inRemoveDuplicates, string inSummaryMetadataCacheDir = null )
         {
             csvFileInfos = new CsvFileInfo[inCsvFilenames.Length];
             for (int i = 0; i < inCsvFilenames.Length; i++)
@@ -2422,6 +2491,7 @@ namespace PerfReportTool
             countFreedSinceLastGC = 0;
 			reportXml = inReportXml;
 			bulkMode = inBulkMode;
+			metadataCacheOnlyMode = inMetadataCacheOnlyMode;
 			metadataQuery = inMetadataQuery;
 			derivedMetadataMappings = inReportXml.derivedMetadataMappings;
 			summaryMetadataCacheDir = inSummaryMetadataCacheDir;
@@ -2429,10 +2499,12 @@ namespace PerfReportTool
 			bRemoveDuplicates = inRemoveDuplicates;
 
 			csvIdToFilename = new Dictionary<string, string>();
+			outCsvQueue = new BlockingCollection<CsvFileInfo>(lookaheadCount);
 
 			// Kick off the workers (must be done last)
 			if (inLookaheadCount > 0)
             {
+				Console.WriteLine("Kicking off " + inThreadCount + " precache threads with lookahead " + lookaheadCount);
                 precacheThreads = new Thread[inThreadCount];
                 precacheJobs = new ThreadStart[inThreadCount];
                 for (int i = 0; i < precacheThreads.Length; i++)
@@ -2442,7 +2514,7 @@ namespace PerfReportTool
                     precacheThreads[i].Start();
                 }
             }
-        }
+		}
 
         public CachedCsvFile GetNextCachedCsvFile()
         {
@@ -2454,10 +2526,10 @@ namespace PerfReportTool
             }
 
 			// Find the next valid fileinfo
-			CsvFileInfo fileInfo = csvFileInfos[readIndex];
 			if (precacheThreads == null)
             {
-                file = new CachedCsvFile(fileInfo.filename, useCacheFiles, derivedMetadataMappings);
+				CsvFileInfo fileInfo = csvFileInfos[readIndex];
+				file = new CachedCsvFile(fileInfo.filename, useCacheFiles, derivedMetadataMappings);
 				if (!file.DoesMetadataMatchQuery(metadataQuery)) // TODO do we need to check this here and in the thread?
 				{
 					return null;
@@ -2465,35 +2537,23 @@ namespace PerfReportTool
 			}
 			else
             {
-                while (true)
-                {
-                    lock (fileInfo.cs)
-                    {
-                        if (fileInfo.isReady)
-                        {
-							file = fileInfo.cachedFile;
-							fileInfo.cachedFile = null;
-							countFreedSinceLastGC++;
-                            // Periodically GC
-                            if (countFreedSinceLastGC>16)
-                            {
-                                GC.Collect();
-                                GC.WaitForPendingFinalizers();
-                                countFreedSinceLastGC = 0;
-                            }
-							if (!fileInfo.isValid)
-							{
-							    Console.WriteLine("Skipping!");
-							    file = null;
-							}
-							break;
-                        }
-                    }
-					// The data isn't ready yet, so sleep for a bit
-                    Thread.Sleep(1);
-                }
-            }
-            readIndex++;
+				CsvFileInfo fileInfo = outCsvQueue.Take();
+				file = fileInfo.cachedFile;
+				fileInfo.cachedFile = null;
+				countFreedSinceLastGC++;
+				// Periodically GC
+				if (countFreedSinceLastGC > 16 && metadataCacheOnlyMode == false)
+				{
+					GC.Collect();
+					GC.WaitForPendingFinalizers();
+					countFreedSinceLastGC = 0;
+				}
+				if (!fileInfo.isValid)
+				{
+					file = null;
+				}
+			}
+			readIndex++;
             return file;
         }
 
@@ -2515,65 +2575,58 @@ namespace PerfReportTool
                 }
                 CsvFileInfo fileInfo = csvFileInfos[threadWriteIndex];
 
-                // If we're too far ahead of the read index, sleep. Not doing so could increase memory usage significantly
-                while (threadWriteIndex - readIndex > lookaheadCount)
-                {
-                    Thread.Sleep(10);
-                }
                 // Process the file
-                lock (fileInfo.cs)
-                {
-					CachedCsvFile file = new CachedCsvFile(fileInfo.filename, useCacheFiles, derivedMetadataMappings);
+				CachedCsvFile file = new CachedCsvFile(fileInfo.filename, useCacheFiles, derivedMetadataMappings);
 
-					// Check for duplicates
-					bool bIsDuplicate = false;
-					if (bRemoveDuplicates && file.metadata != null && file.metadata.Values.ContainsKey("csvid"))
+				// Check for duplicates
+				bool bIsDuplicate = false;
+				if (bRemoveDuplicates && file.metadata != null && file.metadata.Values.ContainsKey("csvid"))
+				{
+					string csvId = file.metadata.Values["csvid"];
+					lock (csvIdToFilename)
 					{
-						string csvId = file.metadata.Values["csvid"];
-						lock (csvIdToFilename)
+						if ( csvIdToFilename.ContainsKey(csvId) )
 						{
-							if ( csvIdToFilename.ContainsKey(csvId) )
-							{
-								Console.WriteLine("Duplicate CSV found: " + fileInfo.filename);
-								Console.WriteLine("   First version   : " + csvIdToFilename[csvId]);
-								bIsDuplicate = true;
-								duplicateCount++;
-							}
-							else
-							{
-								csvIdToFilename[csvId] = fileInfo.filename;
-							}
+							Console.WriteLine("Duplicate CSV found: " + fileInfo.filename);
+							Console.WriteLine("   First version   : " + csvIdToFilename[csvId]);
+							bIsDuplicate = true;
+							duplicateCount++;
+						}
+						else
+						{
+							csvIdToFilename[csvId] = fileInfo.filename;
 						}
 					}
+				}
 
-					if ( !bIsDuplicate && file.DoesMetadataMatchQuery(metadataQuery) )
+				if ( !bIsDuplicate && file.DoesMetadataMatchQuery(metadataQuery) )
+				{
+					file.reportTypeInfo = GetCsvReportTypeInfo(file, bulkMode);
+					if (file.reportTypeInfo != null)
 					{
-						file.reportTypeInfo = GetCsvReportTypeInfo(file, bulkMode);
-						if (file.reportTypeInfo != null)
+						file.ComputeSummaryTableCacheId(file.reportTypeInfo.GetSummaryTableCacheID());
+						if (summaryMetadataCacheDir != null)
 						{
-							file.ComputeSummaryTableCacheId(file.reportTypeInfo.GetSummaryTableCacheID());
-							if (summaryMetadataCacheDir != null)
-							{
-								// If a summary metadata cache is specified, try reading from it instead of reading the whole CSV
-								// Note that this will be disabled if we're not in bulk mode
-								file.cachedSummaryMetadata = SummaryMetadata.TryReadFromCache(summaryMetadataCacheDir, file.summaryTableCacheId);
-								if (file.cachedSummaryMetadata == null)
-								{
-									Console.WriteLine("Failed to read summary metadata from cache for CSV: " + fileInfo.filename);
-								}
-							}
+							// If a summary metadata cache is specified, try reading from it instead of reading the whole CSV
+							// Note that this will be disabled if we're not in bulk mode
+							file.cachedSummaryMetadata = SummaryMetadata.TryReadFromCache(summaryMetadataCacheDir, file.summaryTableCacheId);
 							if (file.cachedSummaryMetadata == null)
 							{
-								file.PrepareCsvData();
+								Console.WriteLine("Failed to read summary metadata from cache for CSV: " + fileInfo.filename);
 							}
-
-							// Only read the full file data if the metadata matches
-							fileInfo.isValid = true;
 						}
+						if (file.cachedSummaryMetadata == null)
+						{
+							file.PrepareCsvData();
+						}
+
+						// Only read the full file data if the metadata matches
+						fileInfo.isValid = true;
 					}
-					fileInfo.cachedFile = file;
-					fileInfo.isReady = true;
 				}
+				fileInfo.cachedFile = file;
+				// TODO: only add valid files to the queue? Need alternative stopping condition
+				outCsvQueue.Add(fileInfo);
 			}
         }
 
@@ -2607,6 +2660,7 @@ namespace PerfReportTool
         int countFreedSinceLastGC;
 		bool useCacheFiles;
 		bool bulkMode;
+		bool metadataCacheOnlyMode;
 		QueryExpression metadataQuery;
 		string summaryMetadataCacheDir;
 		ReportXML reportXml;
@@ -2618,6 +2672,7 @@ namespace PerfReportTool
 		bool bRemoveDuplicates;
 
 		static CsvFileCache fileCache;
+		BlockingCollection<CsvFileInfo> outCsvQueue;
 
         class CsvFileInfo
         {
@@ -2625,13 +2680,9 @@ namespace PerfReportTool
             {
                 filename = inFilename;
                 cachedFile = null;
-                isReady = false;
 				isValid = false;
-				cs = new object();
             }
             public CachedCsvFile cachedFile;
-			public bool isReady;
-            public object cs;
             public string filename;
 			public bool isValid;
         }
diff --git a/Engine/Source/Programs/CSVTools/PerfReportTool/perfsummaries.cs b/Engine/Source/Programs/CSVTools/PerfReportTool/perfsummaries.cs
--- a/Engine/Source/Programs/CSVTools/PerfReportTool/perfsummaries.cs
+++ b/Engine/Source/Programs/CSVTools/PerfReportTool/perfsummaries.cs
@@ -2036,14 +2036,17 @@ namespace PerfSummaries
 
 		public static SummaryMetadata TryReadFromCache(string metadataCacheDir, string csvId)
 		{
-			SummaryMetadata metaData = null;
 			string filename = Path.Combine(metadataCacheDir, csvId + ".prc");
+			return TryReadFromCacheFile(filename);
+		}
 
+		public static SummaryMetadata TryReadFromCacheFile(string filename)
+		{
+			SummaryMetadata metaData = null;
 			if ( !File.Exists(filename) )
 			{
 				return null;
 			}
-
 			try
 			{
 				using (FileStream fileStream = new FileStream(filename, FileMode.Open))
@@ -2078,6 +2081,7 @@ namespace PerfSummaries
 			}
 			return metaData;
 		}
+
 		public bool WriteToCache(string metadataCacheDir, string csvId)
 		{
 			string filename = Path.Combine(metadataCacheDir, csvId + ".prc");
