PerfReportTool 4.33
Open metadata and .bin files readonly

diff --git a/Engine/Source/Programs/CSVTools/CsvStats/CsvStats.cs b/Engine/Source/Programs/CSVTools/CsvStats/CsvStats.cs
--- a/Engine/Source/Programs/CSVTools/CsvStats/CsvStats.cs
+++ b/Engine/Source/Programs/CSVTools/CsvStats/CsvStats.cs
@@ -780,7 +780,7 @@ namespace CSVStats
 
 		public static CsvStats ReadBinFile(string filename, string[] statNamesToRead=null, int numRowsToSkip=0, bool justHeader=false)
 		{
-			System.IO.FileStream fileStream = new FileStream(filename, FileMode.Open);
+			System.IO.FileStream fileStream = new FileStream(filename, FileMode.Open, FileAccess.Read);
 			System.IO.BinaryReader fileReader = new System.IO.BinaryReader(fileStream);
 
 			// Read the header
diff --git a/Engine/Source/Programs/CSVTools/PerfReportTool/PerfReportTool.cs b/Engine/Source/Programs/CSVTools/PerfReportTool/PerfReportTool.cs
--- a/Engine/Source/Programs/CSVTools/PerfReportTool/PerfReportTool.cs
+++ b/Engine/Source/Programs/CSVTools/PerfReportTool/PerfReportTool.cs
@@ -22,7 +22,7 @@ namespace PerfReportTool
 {
     class Version
     {
-        private static string VersionString = "4.32";
+        private static string VersionString = "4.33";
 
         public static string Get() { return VersionString; }
     };
diff --git a/Engine/Source/Programs/CSVTools/PerfReportTool/perfsummaries.cs b/Engine/Source/Programs/CSVTools/PerfReportTool/perfsummaries.cs
--- a/Engine/Source/Programs/CSVTools/PerfReportTool/perfsummaries.cs
+++ b/Engine/Source/Programs/CSVTools/PerfReportTool/perfsummaries.cs
@@ -2049,7 +2049,7 @@ namespace PerfSummaries
 			}
 			try
 			{
-				using (FileStream fileStream = new FileStream(filename, FileMode.Open))
+				using (FileStream fileStream = new FileStream(filename, FileMode.Open, FileAccess.Read))
 				{
 					BinaryReader reader = new BinaryReader(fileStream);
 					int version = reader.ReadInt32();
