PerfReportTool 4.34
 -5x faster metadata cache only reads:
 - Only read the initial CSV metadata entries before metadata filtering (since these always appear first)
 - Skip report type determiniation in metadataonly mode (this is slow and unnecessary)
 -crash fix (don't crash if old version metadata is encountered, just skip and log)

Querying a single device profile across 31k CSVs of mixed device profiles takes 2.1s instead of 10.2s previously

diff --git a/Engine/Source/Programs/CSVTools/PerfReportTool/PerfReportTool.cs b/Engine/Source/Programs/CSVTools/PerfReportTool/PerfReportTool.cs
--- a/Engine/Source/Programs/CSVTools/PerfReportTool/PerfReportTool.cs
+++ b/Engine/Source/Programs/CSVTools/PerfReportTool/PerfReportTool.cs
@@ -22,7 +22,7 @@ namespace PerfReportTool
 {
     class Version
     {
-        private static string VersionString = "4.33";
+        private static string VersionString = "4.34";
 
         public static string Get() { return VersionString; }
     };
@@ -268,7 +268,6 @@ namespace PerfReportTool
 			{
 				throw new Exception("File extension not supported for file " + inFilename);
 			}
-
 			string cacheFilename = inFilename + ".cache";
 			if (useCacheFiles && File.Exists(cacheFilename))
 			{
@@ -305,31 +304,37 @@ namespace PerfReportTool
 				}
 				else if (fileType == FileType.MetadataCachePrc)
 				{
-					// Reconstruct the dummy CSV from the metadata
-					dummyCsvStats = new CsvStats();
-					cachedSummaryMetadata = SummaryMetadata.TryReadFromCacheFile(inFilename);
-					foreach (KeyValuePair<string, SummaryMetadataValue> entry in cachedSummaryMetadata.dict)
+					// Read just the Csv metadata so we can filter
+					metadata = null;
+					cachedSummaryMetadata = SummaryMetadata.TryReadFromCacheFile(inFilename, true);
+					if (cachedSummaryMetadata != null)
 					{
-						if ( entry.Value.type == SummaryMetadataValue.Type.CsvMetadata )
+						foreach (KeyValuePair<string, SummaryMetadataValue> entry in cachedSummaryMetadata.dict)
 						{
-							if (dummyCsvStats.metaData == null)
+							if (entry.Value.type == SummaryMetadataValue.Type.CsvMetadata)
 							{
-								dummyCsvStats.metaData = new CsvMetadata();
+								if (metadata == null)
+								{
+									metadata = new CsvMetadata();
+								}
+								metadata.Values.Add(entry.Key, entry.Value.value);
+							}
+							if (entry.Value.type == SummaryMetadataValue.Type.CsvStatAverage)
+							{
+								break;
 							}
-							dummyCsvStats.metaData.Values.Add(entry.Key, entry.Value.value);
-						}
-						if ( entry.Value.type == SummaryMetadataValue.Type.CsvStatAverage )
-						{
-							dummyCsvStats.AddStat(new StatSamples(entry.Key));
 						}
 					}
-					metadata = dummyCsvStats.metaData;
+					else
+					{
+						Console.WriteLine("Invalid PRC file detected: " + inFilename);
+					}
 				}
 			}
 			filename = inFilename;
 
 			// Apply metadata mappings if we have metadata. Note: we skip for PRCs, since these already have mappings cached
-            if (dummyCsvStats != null && dummyCsvStats.metaData != null && fileType != FileType.MetadataCachePrc)
+            if (dummyCsvStats != null && dummyCsvStats.metaData != null)
             {
 				metadata = dummyCsvStats.metaData;
 				derivedMetadataMappings.ApplyMapping(metadata);
@@ -2601,27 +2606,40 @@ namespace PerfReportTool
 
 				if ( !bIsDuplicate && file.DoesMetadataMatchQuery(metadataQuery) )
 				{
-					file.reportTypeInfo = GetCsvReportTypeInfo(file, bulkMode);
-					if (file.reportTypeInfo != null)
+					if (metadataCacheOnlyMode)
 					{
-						file.ComputeSummaryTableCacheId(file.reportTypeInfo.GetSummaryTableCacheID());
-						if (summaryMetadataCacheDir != null)
+						// Assume the report type is valid if we have cached metadata, since we don't actually need the reporttype in metadata only mode
+						if (file.cachedSummaryMetadata != null)
+						{
+							// Just read the full cached metadata
+							file.cachedSummaryMetadata = SummaryMetadata.TryReadFromCacheFile(fileInfo.filename);
+							fileInfo.isValid = true;
+						}
+					}
+					else
+					{ 
+						file.reportTypeInfo = GetCsvReportTypeInfo(file, bulkMode);
+						if (file.reportTypeInfo != null)
 						{
-							// If a summary metadata cache is specified, try reading from it instead of reading the whole CSV
-							// Note that this will be disabled if we're not in bulk mode
-							file.cachedSummaryMetadata = SummaryMetadata.TryReadFromCache(summaryMetadataCacheDir, file.summaryTableCacheId);
+							file.ComputeSummaryTableCacheId(file.reportTypeInfo.GetSummaryTableCacheID());
+							if (summaryMetadataCacheDir != null)
+							{
+								// If a summary metadata cache is specified, try reading from it instead of reading the whole CSV
+								// Note that this will be disabled if we're not in bulk mode
+								file.cachedSummaryMetadata = SummaryMetadata.TryReadFromCache(summaryMetadataCacheDir, file.summaryTableCacheId);
+								if (file.cachedSummaryMetadata == null)
+								{
+									Console.WriteLine("Failed to read summary metadata from cache for CSV: " + fileInfo.filename);
+								}
+							}
 							if (file.cachedSummaryMetadata == null)
 							{
-								Console.WriteLine("Failed to read summary metadata from cache for CSV: " + fileInfo.filename);
+								file.PrepareCsvData();
 							}
-						}
-						if (file.cachedSummaryMetadata == null)
-						{
-							file.PrepareCsvData();
-						}
 
-						// Only read the full file data if the metadata matches
-						fileInfo.isValid = true;
+							// Only read the full file data if the metadata matches
+							fileInfo.isValid = true;
+						}
 					}
 				}
 				fileInfo.cachedFile = file;
diff --git a/Engine/Source/Programs/CSVTools/PerfReportTool/perfsummaries.cs b/Engine/Source/Programs/CSVTools/PerfReportTool/perfsummaries.cs
--- a/Engine/Source/Programs/CSVTools/PerfReportTool/perfsummaries.cs
+++ b/Engine/Source/Programs/CSVTools/PerfReportTool/perfsummaries.cs
@@ -2032,6 +2032,7 @@ namespace PerfSummaries
 		{
 		}
 
+        // TODO: If this is bumped beyond 6, we need to implement backwards compatibility
 		static int CacheVersion = 6;
 
 		public static SummaryMetadata TryReadFromCache(string metadataCacheDir, string csvId)
@@ -2040,7 +2041,7 @@ namespace PerfSummaries
 			return TryReadFromCacheFile(filename);
 		}
 
-		public static SummaryMetadata TryReadFromCacheFile(string filename)
+		public static SummaryMetadata TryReadFromCacheFile(string filename, bool bReadJustInitialMetadata = false)
 		{
 			SummaryMetadata metaData = null;
 			if ( !File.Exists(filename) )
@@ -2056,19 +2057,30 @@ namespace PerfSummaries
 					int metadataValueVersion = reader.ReadInt32();
 					if (version == CacheVersion && metadataValueVersion == SummaryMetadataValue.CacheVersion)
 					{
+						bool bEarlyOut = false;
 						metaData = new SummaryMetadata();
 						int dictEntryCount = reader.ReadInt32();
 						for (int i = 0; i < dictEntryCount; i++)
 						{
 							string key = reader.ReadString();
 							SummaryMetadataValue value = SummaryMetadataValue.ReadFromCache(reader);
+                            // If we're just reading initial metadata then skip everything after ToolMetadata and CsvMetadata
+							if ( bReadJustInitialMetadata && value.type!=SummaryMetadataValue.Type.ToolMetadata && value.type != SummaryMetadataValue.Type.CsvMetadata )
+							{
+								bEarlyOut = true;
+								break;
+							}
 							metaData.dict.Add(key, value);
 						}
-						string endString = reader.ReadString();
-						if (endString != "END")
+
+						if (!bEarlyOut)
 						{
-							Console.WriteLine("Corruption detected in " + filename + ". Skipping read");
-							metaData = null;
+							string endString = reader.ReadString();
+							if (endString != "END")
+							{
+								Console.WriteLine("Corruption detected in " + filename + ". Skipping read");
+								metaData = null;
+							}
 						}
 					}
 					reader.Close();
