PerfReportTool 4.41
- add -summaryTableCacheUseOnlyCsvID param to ignore the ReportType hash and just use the CSV ID for summaryTableCache ID entries. This is a good idea if you want these files to persist and you don't want them to be regenerated when report types change
- Add ReportType and ReportTypeID to the summaryTable ToolMetaData entries
- Use the reportType name as part of the report hash (this will invalidate existing cache files!)

diff --git a/Engine/Source/Programs/CSVTools/PerfReportTool/PerfReportTool.cs b/Engine/Source/Programs/CSVTools/PerfReportTool/PerfReportTool.cs
--- a/Engine/Source/Programs/CSVTools/PerfReportTool/PerfReportTool.cs
+++ b/Engine/Source/Programs/CSVTools/PerfReportTool/PerfReportTool.cs
@@ -22,7 +22,7 @@ namespace PerfReportTool
 {
     class Version
     {
-        private static string VersionString = "4.40";
+        private static string VersionString = "4.41";
 
         public static string Get() { return VersionString; }
     };
@@ -406,9 +406,16 @@ namespace PerfReportTool
 						sb.Append("{" + key + "}={" + metadata.Values[key] + "}\n");
 					}
 				}
-				csvId = HashHelper.StringToHashStr(sb.ToString()) + "_"; ;
+				csvId = HashHelper.StringToHashStr(sb.ToString());
+			}
+			if ( reportTypeId == null )
+			{
+				summaryTableCacheId = csvId;
+			}
+			else
+			{
+				summaryTableCacheId = csvId + "_" + reportTypeId;
 			}
-			summaryTableCacheId = csvId + "_" + reportTypeId;
 		}
 
 
@@ -512,11 +519,12 @@ namespace PerfReportTool
 			"Performance args for bulk mode:\n" +
 			"       -precacheCount <n> : number of CSV files to precache in the lookahead cache (0 for no precache)\n" +
 			"       -precacheThreads <n> : number of threads to use for the CSV lookahead cache (default 8)\n" +
-			"       -summaryTableCache <dir> : specifies a directory for summary table data to be cached. Enables -readAllStats implicitly.\n  this avoids processing csvs on subsequent runs if -noDetailedReports is specified\n" +
+			"       -summaryTableCache <dir> : specifies a directory for summary table data to be cached. Enables -readAllStats implicitly.\n        This avoids processing csvs on subsequent runs if -noDetailedReports is specified\n" +
 			"       -summaryTableCacheInvalidate : regenerates summary table disk cache entries (ie write only)\n" +
 			"       -summaryTableCacheReadOnly : only read from the cache, never write\n" +
 			"       -summaryTableCachePurgeInvalid : Purges invalid PRCs from the cache folder\n" +
 			"       -summaryTableCacheIn <dir> : reads data directly from the summary table cache instead of from CSVs\n" +
+			"       -summaryTableCacheUseOnlyCsvID : only use the CSV ID for the summary table cache ID. Ignore the report type hash.\n        Use this if you want to avoid cache data being invalidated by report changes\n" +
 			"       -noCsvCacheFiles: disables usage of .csv.cache files. Cache files can be much faster if filtering on metadata\n" +
 			"";
 			/*
@@ -767,7 +775,22 @@ namespace PerfReportTool
 			};
 
 			bool bRemoveDuplicates = !GetBoolArg("allowDuplicateCSVs");
-			CsvFileCache csvFileCache = new CsvFileCache(csvFilenames, precacheCount, precacheThreads, !GetBoolArg("noCsvCacheFiles"), metadataQuery, reportXML, reportTypeParams, bBulkMode, bSummaryTableCacheOnlyMode, bRemoveDuplicates, summaryTableCacheForRead);
+			bool bSummaryTableCacheUseOnlyCsvID = GetBoolArg("summaryTableCacheUseOnlyCsvID");
+
+			CsvFileCache csvFileCache = new CsvFileCache(
+				csvFilenames, 
+				precacheCount, 
+				precacheThreads, 
+				!GetBoolArg("noCsvCacheFiles"), 
+				metadataQuery, 
+				reportXML, 
+				reportTypeParams, 
+				bBulkMode, 
+				bSummaryTableCacheOnlyMode, 
+				bSummaryTableCacheUseOnlyCsvID, 
+				bRemoveDuplicates, 
+				summaryTableCacheForRead );
+
             SummaryTable summaryTable = new SummaryTable();
 			bool bWriteToSummaryTableCache = summaryTableCacheDir != null && !bSummaryTableCacheReadonly;
 
@@ -1176,7 +1199,8 @@ namespace PerfReportTool
 
                 string relativeCsvPath = finalDirUri.MakeRelativeUri(csvFileUri).ToString();
 				rowData.Add(SummaryTableElement.Type.ToolMetadata, "Csv File", "<a href='" + relativeCsvPath + "'>" + shortName + ".csv" + "</a>", null, relativeCsvPath);
-
+				rowData.Add(SummaryTableElement.Type.ToolMetadata, "ReportType", reportTypeInfo.name);
+				rowData.Add(SummaryTableElement.Type.ToolMetadata, "ReportTypeID", reportTypeInfo.summaryTableCacheID);
 				if (htmlFilename != null)
 				{
 					rowData.Add(SummaryTableElement.Type.ToolMetadata, "Report", "<a href='" + htmlFilename + "'>Link</a>");
@@ -1798,7 +1822,8 @@ namespace PerfReportTool
         {	
             graphs = new List<ReportGraph>();
             summaries = new List<Summary>();
-            title = element.Attribute("title").Value;
+			name = element.Attribute("name").Value;
+			title = element.Attribute("title").Value;
             foreach (XElement child in element.Elements())
             {
 				if (child.Name == "graph")
@@ -1864,6 +1889,7 @@ namespace PerfReportTool
 		private void ComputeSummaryTableCacheID()
 		{
 			StringBuilder sb = new StringBuilder();
+			sb.Append("NAME={" + name + "}\n");
 			sb.Append("TITLE={" + title + "}\n");
 			foreach (Summary summary in summaries)
 			{
@@ -1878,7 +1904,8 @@ namespace PerfReportTool
 
 	public List<ReportGraph> graphs;
         public List<Summary> summaries;
-        public string title;
+		public string name;
+		public string title;
 		public string [] metadataToShowList;
 		public string summaryTableCacheID;
 	};
@@ -2481,7 +2508,19 @@ namespace PerfReportTool
 
 	class CsvFileCache
     {
-        public CsvFileCache( string[] inCsvFilenames, int inLookaheadCount, int inThreadCount, bool inUseCacheFiles, QueryExpression inMetadataQuery, ReportXML inReportXml, ReportTypeParams inReportTypeParams, bool inBulkMode, bool inSummaryTableCacheOnlyMode, bool inRemoveDuplicates, string inSummaryTableCacheDir = null )
+        public CsvFileCache( 
+			string[] inCsvFilenames, 
+			int inLookaheadCount, 
+			int inThreadCount, 
+			bool inUseCacheFiles, 
+			QueryExpression inMetadataQuery, 
+			ReportXML inReportXml, 
+			ReportTypeParams inReportTypeParams, 
+			bool inBulkMode, 
+			bool inSummaryTableCacheOnlyMode,
+			bool inSummaryTableCacheUseOnlyCsvID,
+			bool inRemoveDuplicates, 
+			string inSummaryTableCacheDir = null )
         {
             csvFileInfos = new CsvFileInfo[inCsvFilenames.Length];
             for (int i = 0; i < inCsvFilenames.Length; i++)
@@ -2497,6 +2536,7 @@ namespace PerfReportTool
 			reportXml = inReportXml;
 			bulkMode = inBulkMode;
 			summaryTableCacheOnlyMode = inSummaryTableCacheOnlyMode;
+			summaryTableCacheUseOnlyCsvID = inSummaryTableCacheUseOnlyCsvID;
 			metadataQuery = inMetadataQuery;
 			derivedMetadataMappings = inReportXml.derivedMetadataMappings;
 			summaryTableCacheDir = inSummaryTableCacheDir;
@@ -2621,7 +2661,8 @@ namespace PerfReportTool
 						file.reportTypeInfo = GetCsvReportTypeInfo(file, bulkMode);
 						if (file.reportTypeInfo != null)
 						{
-							file.ComputeSummaryTableCacheId(file.reportTypeInfo.GetSummaryTableCacheID());
+							string reportTypeHash = summaryTableCacheUseOnlyCsvID ? null : file.reportTypeInfo.GetSummaryTableCacheID();
+							file.ComputeSummaryTableCacheId(reportTypeHash);
 							if (summaryTableCacheDir != null)
 							{
 								// If a summary metadata cache is specified, try reading from it instead of reading the whole CSV
@@ -2679,6 +2720,7 @@ namespace PerfReportTool
 		bool useCacheFiles;
 		bool bulkMode;
 		bool summaryTableCacheOnlyMode;
+		bool summaryTableCacheUseOnlyCsvID;
 		QueryExpression metadataQuery;
 		string summaryTableCacheDir;
 		ReportXML reportXml;
