PerfReportTool 4.42
- Add -requireMetadata option to ignore CSVs without metadata
- Clean up help text so the lines arent too long for standard console output

diff --git a/Engine/Source/Programs/CSVTools/PerfReportTool/PerfReportTool.cs b/Engine/Source/Programs/CSVTools/PerfReportTool/PerfReportTool.cs
--- a/Engine/Source/Programs/CSVTools/PerfReportTool/PerfReportTool.cs
+++ b/Engine/Source/Programs/CSVTools/PerfReportTool/PerfReportTool.cs
@@ -22,7 +22,7 @@ namespace PerfReportTool
 {
     class Version
     {
-        private static string VersionString = "4.41";
+        private static string VersionString = "4.42";
 
         public static string Get() { return VersionString; }
     };
@@ -472,22 +472,25 @@ namespace PerfReportTool
 			"       -reportxmlbasedir <folder>\n" +
 			"       -title <name>\n" +
 			"       -maxy <value> - forces all graphs to use this value\n" +
-			"       -writeSummaryCsv : if specified, a csv file containing summary information will be generated. Not available in bulk mode.\n" +
-			"       -noWatermarks : don't embed the commandline or version in reports"+
-			"       -cleanCsvOut <filename> : write a standard format CSV after event stripping with metadata stripped out. Non-bulk mode only\n"+
-			"\n"+
+			"       -writeSummaryCsv : if specified, a csv file containing summary information will be generated.\n"+
+			"          Not available in bulk mode.\n" +
+			"       -noWatermarks : don't embed the commandline or version in reports\n"+
+			"       -cleanCsvOut <filename> : write a standard format CSV after event stripping with metadata stripped out.\n"+
+			"          Not available in bulk mode.\n" +
+			"\n" +
 			"Performance args:\n" +
 			"       -perfLog : output performance logging information\n" +
 			"       -noBatchedGraphs : disable batched/multithreaded graph generation (default is enabled)\n" +
-			"       -graphThreads : use with -batchedGraphs to control the number of threads per CsvToSVG instance (default: PC core count/2)\n" +
+			"       -graphThreads : use with -batchedGraphs to control the number of threads per CsvToSVG instance \n"+
+			"                       (default: PC core count/2)\n" +
 			"\n" +
 			"Options to truncate or filter source data:\n" +
+			"Warning: these options disable Summary Table caching\n" +
 			"       -minx <frameNumber>\n" +
 			"       -maxx <frameNumber>\n" +
 			"       -beginEvent <event> : strip data before this event\n" +
 			"       -endEvent <event> : strip data after this event\n" +
 			"       -noStripEvents : if specified, don't strip out samples between excluded events from the stats\n" +
-			"NOTE: these options disable Summary Table caching\n" +
 			"\n" +
 			"Optional bulk mode args: (use with -csvdir or -summaryTableCacheIn)\n" +
 			"       -recurse \n" +
@@ -505,9 +508,10 @@ namespace PerfReportTool
 			"           Selects a custom condensed summary table type from the list in reportTypes.xml \n"+
 			"           (if not specified, 'condensed' will be used)\n" +
 			"       -summaryTableFilename <name> : use the specified filename for the summary table (instead of SummaryTable.html)\n"+
-			"       -metadataFilter <query> or <key0=value0,key1=value1...>: filters based on CSV metadata, e.g \"platform=ps4 AND deviceprofile=ps4_60\" \n" +
-			"       -readAllStats : reads all stats so that any stat can be output to the summary table. Useful with -customtable in bulk mode (off by default)\n" +
-			"       -showHiddenStats : shows stats which have been automatically hidden (e.g csv unit stat averages hidden by FPSCharts Summary)\n" +
+			"       -metadataFilter <query> or <key0=value0,key1=value1...>: filters based on CSV metadata,\n"+
+			"           e.g \"platform=ps4 AND deviceprofile=ps4_60\" \n" +
+			"       -readAllStats : allows any CSV stat avg to appear in the summary table, not just those referenced in summaries\n" +
+			"       -showHiddenStats : shows stats which have been automatically hidden (typically duplicate csv unit stats)\n" +
 			"       -externalGraphs : enables external graphs (off by default)\n" +
 			"       -spreadsheetfriendly: outputs a single quote before non-numeric entries in summary tables\n" +
 			"       -noSummaryMinMax: don't make min/max columns for each stat in a condensed summary\n" +
@@ -515,16 +519,20 @@ namespace PerfReportTool
 			"       -scrollableTable: makes the summary table scrollable, with frozen first rows and columns\n" +
 			"       -maxSummaryTableStringLength <n>: strings longer than this will get truncated\n" +
 			"       -allowDuplicateCSVs : doesn't remove duplicate CSVs (Note: can cause summary table cache file locking issues)\n"+
+			"       -requireMetadata : ignores CSVs without metadata\n" +
 			"\n" +
 			"Performance args for bulk mode:\n" +
 			"       -precacheCount <n> : number of CSV files to precache in the lookahead cache (0 for no precache)\n" +
 			"       -precacheThreads <n> : number of threads to use for the CSV lookahead cache (default 8)\n" +
-			"       -summaryTableCache <dir> : specifies a directory for summary table data to be cached. Enables -readAllStats implicitly.\n        This avoids processing csvs on subsequent runs if -noDetailedReports is specified\n" +
+			"       -summaryTableCache <dir> : specifies a directory for summary table data to be cached.\n"+
+			"           This avoids processing csvs on subsequent runs when -noDetailedReports is specified\n" +
+			"           Note: Enables -readAllStats implicitly. \n"+
 			"       -summaryTableCacheInvalidate : regenerates summary table disk cache entries (ie write only)\n" +
 			"       -summaryTableCacheReadOnly : only read from the cache, never write\n" +
 			"       -summaryTableCachePurgeInvalid : Purges invalid PRCs from the cache folder\n" +
 			"       -summaryTableCacheIn <dir> : reads data directly from the summary table cache instead of from CSVs\n" +
-			"       -summaryTableCacheUseOnlyCsvID : only use the CSV ID for the summary table cache ID. Ignore the report type hash.\n        Use this if you want to avoid cache data being invalidated by report changes\n" +
+			"       -summaryTableCacheUseOnlyCsvID : only use the CSV ID for the summary table cacheID, ignoringthe report type hash\n"+
+			"            Use this if you want to avoid cache data being invalidated by report changes\n" +
 			"       -noCsvCacheFiles: disables usage of .csv.cache files. Cache files can be much faster if filtering on metadata\n" +
 			"";
 			/*
@@ -776,6 +784,7 @@ namespace PerfReportTool
 
 			bool bRemoveDuplicates = !GetBoolArg("allowDuplicateCSVs");
 			bool bSummaryTableCacheUseOnlyCsvID = GetBoolArg("summaryTableCacheUseOnlyCsvID");
+			bool bRequireMetadata = GetBoolArg("requireMetadata");
 
 			CsvFileCache csvFileCache = new CsvFileCache(
 				csvFilenames, 
@@ -788,7 +797,8 @@ namespace PerfReportTool
 				bBulkMode, 
 				bSummaryTableCacheOnlyMode, 
 				bSummaryTableCacheUseOnlyCsvID, 
-				bRemoveDuplicates, 
+				bRemoveDuplicates,
+				bRequireMetadata,
 				summaryTableCacheForRead );
 
             SummaryTable summaryTable = new SummaryTable();
@@ -2520,7 +2530,8 @@ namespace PerfReportTool
 			bool inSummaryTableCacheOnlyMode,
 			bool inSummaryTableCacheUseOnlyCsvID,
 			bool inRemoveDuplicates, 
-			string inSummaryTableCacheDir = null )
+			bool inRequireMetadata,
+			string inSummaryTableCacheDir )
         {
             csvFileInfos = new CsvFileInfo[inCsvFilenames.Length];
             for (int i = 0; i < inCsvFilenames.Length; i++)
@@ -2542,6 +2553,7 @@ namespace PerfReportTool
 			summaryTableCacheDir = inSummaryTableCacheDir;
 			reportTypeParams = inReportTypeParams;
 			bRemoveDuplicates = inRemoveDuplicates;
+			bRequireMetadata = inRequireMetadata;
 
 			csvIdToFilename = new Dictionary<string, string>();
 			outCsvQueue = new BlockingCollection<CsvFileInfo>(lookaheadCount);
@@ -2623,9 +2635,13 @@ namespace PerfReportTool
                 // Process the file
 				CachedCsvFile file = new CachedCsvFile(fileInfo.filename, useCacheFiles, derivedMetadataMappings);
 
-				// Check for duplicates
-				bool bIsDuplicate = false;
-				if (bRemoveDuplicates && file.metadata != null && file.metadata.Values.ContainsKey("csvid"))
+				bool bProcessThisFile = true;
+				if (bRequireMetadata && file.metadata == null)
+				{
+					Console.WriteLine("CSV has no metadata. Skipping: " + fileInfo.filename);
+					bProcessThisFile = false;
+				}
+				else if (bRemoveDuplicates && file.metadata != null && file.metadata.Values.ContainsKey("csvid"))
 				{
 					string csvId = file.metadata.Values["csvid"];
 					lock (csvIdToFilename)
@@ -2634,7 +2650,7 @@ namespace PerfReportTool
 						{
 							Console.WriteLine("Duplicate CSV found: " + fileInfo.filename);
 							Console.WriteLine("   First version   : " + csvIdToFilename[csvId]);
-							bIsDuplicate = true;
+							bProcessThisFile = false;
 							duplicateCount++;
 						}
 						else
@@ -2644,7 +2660,7 @@ namespace PerfReportTool
 					}
 				}
 
-				if ( !bIsDuplicate && file.DoesMetadataMatchQuery(metadataQuery) )
+				if ( bProcessThisFile && file.DoesMetadataMatchQuery(metadataQuery) )
 				{
 					if (summaryTableCacheOnlyMode)
 					{
@@ -2719,6 +2735,7 @@ namespace PerfReportTool
         int countFreedSinceLastGC;
 		bool useCacheFiles;
 		bool bulkMode;
+		bool bRequireMetadata;
 		bool summaryTableCacheOnlyMode;
 		bool summaryTableCacheUseOnlyCsvID;
 		QueryExpression metadataQuery;
