PerfReportTool 4.43
- Add support for >,<,>=,<= operators for metadata queries (for numerical fields)
- Use doubles for SummaryTable/SummaryTableColumn values, so we can accurately show display values (does not affect summary cache - this was already at double precision)

diff --git a/Engine/Source/Programs/CSVTools/PerfReportTool/PerfReportTool.cs b/Engine/Source/Programs/CSVTools/PerfReportTool/PerfReportTool.cs
--- a/Engine/Source/Programs/CSVTools/PerfReportTool/PerfReportTool.cs
+++ b/Engine/Source/Programs/CSVTools/PerfReportTool/PerfReportTool.cs
@@ -22,7 +22,7 @@ namespace PerfReportTool
 {
     class Version
     {
-        private static string VersionString = "4.42";
+        private static string VersionString = "4.43";
 
         public static string Get() { return VersionString; }
     };
diff --git a/Engine/Source/Programs/CSVTools/PerfReportTool/perfsummaries.cs b/Engine/Source/Programs/CSVTools/PerfReportTool/perfsummaries.cs
--- a/Engine/Source/Programs/CSVTools/PerfReportTool/perfsummaries.cs
+++ b/Engine/Source/Programs/CSVTools/PerfReportTool/perfsummaries.cs
@@ -2221,7 +2221,7 @@ namespace PerfSummaries
 		public string name;
 		public bool isNumeric = false;
 		public string displayName;
-		List<float> floatValues = new List<float>();
+		List<double> doubleValues = new List<double>();
 		List<string> stringValues = new List<string>();
 		List<string> toolTips = new List<string>();
 
@@ -2236,7 +2236,7 @@ namespace PerfSummaries
 		public SummaryTableColumn Clone()
 		{
 			SummaryTableColumn newColumn = new SummaryTableColumn(name, isNumeric, displayName);
-			newColumn.floatValues.AddRange(floatValues);
+			newColumn.doubleValues.AddRange(doubleValues);
 			newColumn.stringValues.AddRange(stringValues);
 			newColumn.colourThresholds.AddRange(colourThresholds);
 			newColumn.toolTips.AddRange(toolTips);
@@ -2252,7 +2252,7 @@ namespace PerfSummaries
 			return displayName;
 		}
 
-		public void SetValue(int index, float value)
+		public void SetValue(int index, double value)
 		{
 			if ( !isNumeric )
 			{
@@ -2261,14 +2261,14 @@ namespace PerfSummaries
 				return;
 			}
 			// Grow to fill if necessary
-			if ( index >= floatValues.Count )
+			if ( index >= doubleValues.Count )
 			{
-				for ( int i= floatValues.Count; i<=index; i++ )
+				for ( int i= doubleValues.Count; i<=index; i++ )
 				{
-					floatValues.Add(float.MaxValue);
+					doubleValues.Add(double.MaxValue);
 				}
 			}
-			floatValues[index] = value;
+			doubleValues[index] = value;
 		}
 
 		void convertToStrings()
@@ -2276,11 +2276,11 @@ namespace PerfSummaries
 			if ( isNumeric )
 			{
 				stringValues = new List<string>();
-				foreach (float f in floatValues)
+				foreach (float f in doubleValues)
 				{
 					stringValues.Add(f.ToString());
 				}
-				floatValues = new List<float>();
+				doubleValues = new List<double>();
 				isNumeric = false;
 			}
 		}
@@ -2310,8 +2310,8 @@ namespace PerfSummaries
 		public string GetColour(int index)
 		{
 			ColourThresholdList thresholds = null;
-			float value = GetValue(index);
-			if (value==float.MaxValue)
+			double value = GetValue(index);
+			if (value == double.MaxValue)
 			{
 				return null;
 			}
@@ -2330,20 +2330,20 @@ namespace PerfSummaries
 					return null;
 				}
 			}
-			return thresholds.GetColourForValue((double)value);
+			return thresholds.GetColourForValue(value);
 		}
 
 		public void ComputeAutomaticColourThresholds(bool bHighIsRed)
 		{
 			colourThresholds = new List<ColourThresholdList>();
-			float maxValue = -float.MaxValue;
-			float minValue = float.MaxValue;
-			float totalValue = 0.0f;
-			float validCount = 0.0f;
-			for ( int i=0; i<floatValues.Count; i++ )
+			double maxValue = -double.MaxValue;
+			double minValue = double.MaxValue;
+			double totalValue = 0.0f;
+			double validCount = 0.0f;
+			for ( int i=0; i< doubleValues.Count; i++ )
 			{
-				float val = floatValues[i];
-				if (val != float.MaxValue)
+				double val = doubleValues[i];
+				if (val != double.MaxValue)
 				{
 					maxValue = Math.Max(val, maxValue);
 					minValue = Math.Min(val, minValue);
@@ -2360,7 +2360,7 @@ namespace PerfSummaries
 			Colour yellow = new Colour(1.0f, 1.0f, 0.5f);
 			Colour red = new Colour(1.0f, 0.4f, 0.4f);
 
-			float averageValue = totalValue / validCount; // TODO: Weighted average 
+			double averageValue = totalValue / validCount; // TODO: Weighted average 
 			colourThresholdOverride = new ColourThresholdList();
 			colourThresholdOverride.Add(new ThresholdInfo(minValue, bHighIsRed ? green : red));
 			colourThresholdOverride.Add(new ThresholdInfo(averageValue, yellow));
@@ -2370,15 +2370,15 @@ namespace PerfSummaries
 
 		public int GetCount()
 		{
-			return Math.Max(floatValues.Count, stringValues.Count);
+			return Math.Max(doubleValues.Count, stringValues.Count);
 		}
-		public float GetValue(int index)
+		public double GetValue(int index)
 		{
-			if ( index >= floatValues.Count )
+			if ( index >= doubleValues.Count )
 			{
-				return float.MaxValue;
+				return double.MaxValue;
 			}
-			return floatValues[index];
+			return doubleValues[index];
 		}
 
 		public void SetStringValue(int index, string value)
@@ -2403,15 +2403,15 @@ namespace PerfSummaries
 		{
 			if (isNumeric)
 			{
-				if (index >= floatValues.Count || floatValues[index] == float.MaxValue)
+				if (index >= doubleValues.Count || doubleValues[index] == double.MaxValue)
 				{
 					return "";
 				}
-				float val = floatValues[index];
+				double val = doubleValues[index];
 				if (roundNumericValues)
 				{
-					float absVal = Math.Abs(val);
-					float frac = absVal - (float)Math.Truncate(absVal);
+					double absVal = Math.Abs(val);
+					double frac = absVal - (double)Math.Truncate(absVal);
 					if (absVal >= 250.0f || frac < 0.0001f )
 					{
 						return val.ToString("0");
@@ -2513,9 +2513,9 @@ namespace PerfSummaries
 				}
 			}
 
-			List<float> RowMaxValues = new List<float>();
-			List<float> RowTotals = new List<float>();
-			List<float> RowMinValues = new List<float>();
+			List<double> RowMaxValues = new List<double>();
+			List<double> RowTotals = new List<double>();
+			List<double> RowMinValues = new List<double>();
 			List<int> RowCounts = new List<int>();
 			List<ColourThresholdList> RowColourThresholds = new List<ColourThresholdList>();
 
@@ -2542,8 +2542,8 @@ namespace PerfSummaries
 					{
 						if (addMinMaxColumns)
 						{
-							RowMaxValues.Add(-float.MaxValue);
-							RowMinValues.Add(float.MaxValue);
+							RowMaxValues.Add(-double.MaxValue);
+							RowMinValues.Add(double.MaxValue);
 						}
 						RowTotals.Add(0.0f);
 						RowCounts.Add(0);
@@ -2559,8 +2559,8 @@ namespace PerfSummaries
 					SummaryTableColumn column = columns[j];
 					if (column.isNumeric)
 					{
-						float value = column.GetValue(i);
-						if (value != float.MaxValue)
+						double value = column.GetValue(i);
+						if (value != double.MaxValue)
 						{
 							if (addMinMaxColumns)
 							{
@@ -2594,13 +2594,13 @@ namespace PerfSummaries
 						newColumns[j].SetStringValue(destRowIndex, columnLookup[key].GetStringValue(i));
 					}
 					// Commit the row 
-					newColumns[countColumnIndex].SetValue(destRowIndex, (float)mergedRowsCount);
+					newColumns[countColumnIndex].SetValue(destRowIndex, (double)mergedRowsCount);
 					for (int j = 0; j < columns.Count; j++)
 					{
 						int destColumnBaseIndex = srcToDestBaseColumnIndex[j];
 						if (destColumnBaseIndex != -1 && RowCounts[j]>0)
 						{
-							newColumns[destColumnBaseIndex].SetValue(destRowIndex, RowTotals[j] / (float)RowCounts[j]);
+							newColumns[destColumnBaseIndex].SetValue(destRowIndex, RowTotals[j] / (double)RowCounts[j]);
 							if (addMinMaxColumns)
 							{
 								newColumns[destColumnBaseIndex + 1].SetValue(destRowIndex, RowMinValues[j]);
@@ -3201,7 +3201,7 @@ namespace PerfSummaries
 
 				if (value.isNumeric)
 				{
-					column.SetValue(rowCount, (float)value.numericValue);
+					column.SetValue(rowCount, (double)value.numericValue);
 				}
 				else
 				{
