Add missing file from CL 16110876

diff --git a/Engine/Source/Programs/CSVTools/CsvStats/MetadataQueryBuilder.cs b/Engine/Source/Programs/CSVTools/CsvStats/MetadataQueryBuilder.cs
--- a/Engine/Source/Programs/CSVTools/CsvStats/MetadataQueryBuilder.cs
+++ b/Engine/Source/Programs/CSVTools/CsvStats/MetadataQueryBuilder.cs
@@ -17,6 +17,10 @@ namespace CSVStats
 	{
 		Equal,
 		NotEqual,
+		GreaterThan,
+		LessThan,
+		GreaterThanOrEqual,
+		LessThanOrEqual,
 
 		COUNT
 	};
@@ -36,12 +40,42 @@ namespace CSVStats
 			{
 				return false;
 			}
-			// Check if the value actually matches (allow wildcards)
-			bool matches=CsvStats.DoesSearchStringMatch(metadata.Values[Key].ToLower(), Value.ToLower());
+			bool matches = false;
 			bool negate = bNegate;
-			if (Op == ComparisonOp.NotEqual)
+			if (Op == ComparisonOp.Equal || Op == ComparisonOp.NotEqual)
 			{
-				negate = !negate;
+				// Check if the value actually matches (allow wildcards)
+				matches = CsvStats.DoesSearchStringMatch(metadata.Values[Key].ToLower(), Value.ToLower());
+				if (Op == ComparisonOp.NotEqual)
+				{
+					negate = !negate;
+				}
+			}
+			if (Op == ComparisonOp.LessThan || Op == ComparisonOp.GreaterThan || Op == ComparisonOp.LessThanOrEqual || Op == ComparisonOp.GreaterThanOrEqual )
+			{
+				double doubleVal;
+				if ( !Double.TryParse(Value, out doubleVal ) )
+				{
+					throw new Exception("Metadata operators expect numeric values!");
+				}
+				double doubleMetadataVal;
+				if (!Double.TryParse(metadata.Values[Key].ToLower(), out doubleMetadataVal))
+				{
+					// Metadata value wasn't numeric, so just return
+					Console.WriteLine("Value for "+Key+" was non-numeric: " + Value);
+					return false;
+				}
+				switch (Op)
+				{
+					case ComparisonOp.LessThan:
+						return doubleMetadataVal < doubleVal;
+					case ComparisonOp.GreaterThan:
+						return doubleMetadataVal > doubleVal;
+					case ComparisonOp.LessThanOrEqual:
+						return doubleMetadataVal <= doubleVal;
+					case ComparisonOp.GreaterThanOrEqual:
+						return doubleMetadataVal >= doubleVal;
+				}
 			}
 			return negate ? !matches : matches;
 		}
@@ -96,11 +130,35 @@ namespace CSVStats
 			{
 				return null;
 			}
-			if ( tokenList[1] == "=" || tokenList[1] == "!=")
+			ComparisonOp op = ComparisonOp.COUNT;
+			switch (tokenList[1])
+			{
+				case "=":
+					op = ComparisonOp.Equal;
+					break;
+				case "!=":
+					op = ComparisonOp.NotEqual;
+					break;
+				case "<":
+					op = ComparisonOp.LessThan;
+					break;
+				case ">":
+					op = ComparisonOp.GreaterThan;
+					break;
+				case "<=":
+					op = ComparisonOp.LessThanOrEqual;
+					break;
+				case ">=":
+					op = ComparisonOp.GreaterThanOrEqual;
+					break;
+				default:
+					break;
+			}
+			if ( op != ComparisonOp.COUNT )
 			{
 				termOut.Key = tokenList[0];
 				termOut.Value = tokenList[2];
-				termOut.Op = tokenList[1] == "=" ? ComparisonOp.Equal : ComparisonOp.NotEqual;
+				termOut.Op = op;
 				tokenList.RemoveRange(0, 3);
 				return termOut;
 			}
@@ -234,19 +292,75 @@ namespace CSVStats
 			Whitespace,
 		}
 
+		private static string ConsumeNextToken(ref string remainingString)
+		{
+			string[] operatorTokenArray = { "!=", "<=", ">=", "(", ")", "=", "<", ">" };
+			string operatorChars = "!=<>()";
+
+			// Consume leading whitespace
+			remainingString = remainingString.TrimStart();
+			if (remainingString.Length == 0)
+			{
+				return null;
+			}
+
+			// Consume operators
+			foreach (string opString in operatorTokenArray)
+			{
+				if (remainingString.StartsWith(opString))
+				{
+					remainingString = remainingString.Substring(opString.Length);
+					return opString;
+				}
+			}
+			// Not an operator or whitespace
+			string token = "";
+			for (int i = 0; i < remainingString.Length; i++)
+			{
+				char c = remainingString[i];
+				if (Char.IsWhiteSpace(c) || operatorChars.IndexOf(c) >= 0)
+				{
+					break;
+				}
+				token += c;
+			}
+			remainingString = remainingString.Substring(token.Length);
+			return token;
+		}
+
 		private static List<string> Tokenize(string str)
 		{
 			// Tokenize the string
-			str = str.Trim() + " ";
-			string[] operatorTokenArray = { "!=", "(", ")", "=" };
+			string remainingString = str;
+			List<string> tokens = new List<string>();
+			while (remainingString.Length > 0)
+			{
+				string token = ConsumeNextToken(ref remainingString);
+				if (token != null)
+				{
+					tokens.Add(token);
+				}
+			}
+
+
+			/*
 			List<string> operatorTokens = new List<string>(operatorTokenArray);
-			string operatorChars = "!=()";
+
+			// Find all operator chars
+			string operatorChars = "";
+			foreach (string opToken in operatorTokenArray)
+				foreach (char c in opToken)
+					if (operatorChars.IndexOf(c)<0)
+						operatorChars += c;
+
 			List<string> tokens = new List<string>();
 			string currentToken = "";
-			CharType currentCharType = CharType.Whitespace;
+			CharType prevCharType = CharType.Whitespace;
 			for (int i = 0; i < str.Length; i++)
 			{
 				char c = str[i];
+
+				// Find the character type
 				CharType newCharType;
 				if (Char.IsWhiteSpace(c))
 				{
@@ -255,31 +369,32 @@ namespace CSVStats
 				else if (operatorChars.Contains(c.ToString()))
 				{
 					newCharType = CharType.Operator;
-					// If we've already got an operator then add it now
-					if (operatorTokens.Contains(currentToken))
-					{
-						tokens.Add(currentToken);
-						currentToken = "";
-					}
 				}
 				else
 				{
 					newCharType = CharType.Standard;
 				}
-				if (newCharType != currentCharType)
+
+				// If the character type changed, add the token
+				if (newCharType != prevCharType)
 				{
-					currentCharType = newCharType;
+					prevCharType = newCharType;
 					if (currentToken.Length > 0)
 					{
 						tokens.Add(currentToken);
+						currentToken = "";
 					}
-					currentToken = "";
 				}
 				if ( newCharType != CharType.Whitespace )
 				{
 					currentToken += c;
 				}
 			}
+			if (currentToken.Length > 0)
+			{
+				tokens.Add(currentToken);
+			}
+			*/
 			return tokens;
 		}
 
