PerfReportTool 4.44
-add -reportLinkRootPath arg to make report links relative to a particular location
-add -csvLinkRootPath arg to make CSV links relative to a particular location
-add -listFiles arg to just list files that pass the filter instead of processing them

diff --git a/Engine/Source/Programs/CSVTools/PerfReportTool/PerfReportTool.cs b/Engine/Source/Programs/CSVTools/PerfReportTool/PerfReportTool.cs
--- a/Engine/Source/Programs/CSVTools/PerfReportTool/PerfReportTool.cs
+++ b/Engine/Source/Programs/CSVTools/PerfReportTool/PerfReportTool.cs
@@ -22,7 +22,7 @@ namespace PerfReportTool
 {
     class Version
     {
-        private static string VersionString = "4.43";
+        private static string VersionString = "4.44";
 
         public static string Get() { return VersionString; }
     };
@@ -520,6 +520,9 @@ namespace PerfReportTool
 			"       -maxSummaryTableStringLength <n>: strings longer than this will get truncated\n" +
 			"       -allowDuplicateCSVs : doesn't remove duplicate CSVs (Note: can cause summary table cache file locking issues)\n"+
 			"       -requireMetadata : ignores CSVs without metadata\n" +
+			"       -listFiles : just list all files that pass the metadata query. Don't generate any reports.\n" +
+			"       -reportLinkRootPath : Make report links relative to this\n" +
+			"       -csvLinkRootPath : Make CSV file links relative to this\n" +
 			"\n" +
 			"Performance args for bulk mode:\n" +
 			"       -precacheCount <n> : number of CSV files to precache in the lookahead cache (0 for no precache)\n" +
@@ -785,6 +788,12 @@ namespace PerfReportTool
 			bool bRemoveDuplicates = !GetBoolArg("allowDuplicateCSVs");
 			bool bSummaryTableCacheUseOnlyCsvID = GetBoolArg("summaryTableCacheUseOnlyCsvID");
 			bool bRequireMetadata = GetBoolArg("requireMetadata");
+			bool bListFilesMode = GetBoolArg("listFiles");
+			if (bListFilesMode)
+			{
+				writeDetailedReports = false;
+			}
+
 
 			CsvFileCache csvFileCache = new CsvFileCache(
 				csvFilenames, 
@@ -799,7 +808,8 @@ namespace PerfReportTool
 				bSummaryTableCacheUseOnlyCsvID, 
 				bRemoveDuplicates,
 				bRequireMetadata,
-				summaryTableCacheForRead );
+				summaryTableCacheForRead,
+				bListFilesMode);
 
             SummaryTable summaryTable = new SummaryTable();
 			bool bWriteToSummaryTableCache = summaryTableCacheDir != null && !bSummaryTableCacheReadonly;
@@ -1195,6 +1205,8 @@ namespace PerfReportTool
                     outputDir += "/";
                 }
                 Uri optionalDirUri = new Uri(outputDir, UriKind.RelativeOrAbsolute);
+
+				// Make a Csv URI that's relative to the report directory
                 Uri finalDirUri;
                 if (optionalDirUri.IsAbsoluteUri)
                 {
@@ -1205,14 +1217,27 @@ namespace PerfReportTool
                     finalDirUri = new Uri(currentDirUri,outputDir);
                 }
                 Uri csvFileUri = new Uri(csvFile.filename, UriKind.Absolute);
+				Uri relativeCsvUri = finalDirUri.MakeRelativeUri(csvFileUri);
+				string csvPath = relativeCsvUri.ToString();
 
-                string relativeCsvPath = finalDirUri.MakeRelativeUri(csvFileUri).ToString();
-				rowData.Add(SummaryTableElement.Type.ToolMetadata, "Csv File", "<a href='" + relativeCsvPath + "'>" + shortName + ".csv" + "</a>", null, relativeCsvPath);
+				string csvLinkRootPath = GetArg("csvLinkRootPath", null);
+				if ( csvLinkRootPath != null && !relativeCsvUri.IsAbsoluteUri)
+				{
+					csvPath = csvLinkRootPath + csvPath;
+				}
+				rowData.Add(SummaryTableElement.Type.ToolMetadata, "Csv File", "<a href='" + csvPath + "'>" + shortName + ".csv" + "</a>", null, csvPath);
 				rowData.Add(SummaryTableElement.Type.ToolMetadata, "ReportType", reportTypeInfo.name);
 				rowData.Add(SummaryTableElement.Type.ToolMetadata, "ReportTypeID", reportTypeInfo.summaryTableCacheID);
 				if (htmlFilename != null)
 				{
-					rowData.Add(SummaryTableElement.Type.ToolMetadata, "Report", "<a href='" + htmlFilename + "'>Link</a>");
+					string htmlUrl = htmlFilename;
+					string reportLinkRootPath = GetArg("reportLinkRootPath", null);
+					if (reportLinkRootPath != null)
+					{
+						htmlUrl = reportLinkRootPath + htmlFilename;
+					}
+
+					rowData.Add(SummaryTableElement.Type.ToolMetadata, "Report", "<a href='" + htmlUrl + "'>Link</a>");
 				}
 				// Pass through all the metadata from the CSV
 				if (csvStats.metaData != null)
@@ -2532,7 +2557,8 @@ namespace PerfReportTool
 			bool inSummaryTableCacheUseOnlyCsvID,
 			bool inRemoveDuplicates, 
 			bool inRequireMetadata,
-			string inSummaryTableCacheDir )
+			string inSummaryTableCacheDir,
+			bool inListFilesMode)
         {
             csvFileInfos = new CsvFileInfo[inCsvFilenames.Length];
             for (int i = 0; i < inCsvFilenames.Length; i++)
@@ -2555,6 +2581,7 @@ namespace PerfReportTool
 			reportTypeParams = inReportTypeParams;
 			bRemoveDuplicates = inRemoveDuplicates;
 			bRequireMetadata = inRequireMetadata;
+			bListFilesMode = inListFilesMode;
 
 			csvIdToFilename = new Dictionary<string, string>();
 			outCsvQueue = new BlockingCollection<CsvFileInfo>(lookaheadCount);
@@ -2663,7 +2690,11 @@ namespace PerfReportTool
 
 				if ( bProcessThisFile && file.DoesMetadataMatchQuery(metadataQuery) )
 				{
-					if (summaryTableCacheOnlyMode)
+					if (bListFilesMode)
+					{
+						Console.WriteLine("File: " + fileInfo.filename);
+					}
+					else if (summaryTableCacheOnlyMode)
 					{
 						// Assume the report type is valid if we have cached metadata, since we don't actually need the reporttype in metadata only mode
 						if (file.cachedSummaryTableRowData != null)
@@ -2748,6 +2779,7 @@ namespace PerfReportTool
 		Dictionary<string, string> csvIdToFilename;
 		public int duplicateCount=0;
 		bool bRemoveDuplicates;
+		bool bListFilesMode;
 
 		static CsvFileCache fileCache;
 		BlockingCollection<CsvFileInfo> outCsvQueue;
