PerfReportTool 4.45
- make -csvLinkRootPath work any type of path. Links are made relative to the passed-in csv directory and then the new root is added

diff --git a/Engine/Source/Programs/CSVTools/PerfReportTool/PerfReportTool.cs b/Engine/Source/Programs/CSVTools/PerfReportTool/PerfReportTool.cs
--- a/Engine/Source/Programs/CSVTools/PerfReportTool/PerfReportTool.cs
+++ b/Engine/Source/Programs/CSVTools/PerfReportTool/PerfReportTool.cs
@@ -22,7 +22,7 @@ namespace PerfReportTool
 {
     class Version
     {
-        private static string VersionString = "4.44";
+        private static string VersionString = "4.45";
 
         public static string Get() { return VersionString; }
     };
@@ -581,8 +581,9 @@ namespace PerfReportTool
 
 			ReadCommandLine(args);
             PerfLog perfLog = new PerfLog(GetBoolArg("perfLog"));
+			string csvDir = null;
 
-            bool bBulkMode = false;
+			bool bBulkMode = false;
 			bool bSummaryTableCacheOnlyMode = false;
 			// Read CSV filenames from a directory or list
 			string[] csvFilenames;
@@ -593,7 +594,7 @@ namespace PerfReportTool
 			}
 			else
 			{
-				string csvDir = GetArg("csvDir");
+				csvDir=GetArg("csvDir");
 				string summaryTableCacheInDir = GetArg("summaryTableCacheIn");
 				if (csvDir.Length > 0)
 				{
@@ -853,7 +854,7 @@ namespace PerfReportTool
 							}
 							else
 							{
-								GenerateReport(cachedCsvFile, outputDir, bBulkMode, rowData, bBatchedGraphs, writeDetailedReports, bReadAllStats || bWriteToSummaryTableCache, cachedCsvFile.reportTypeInfo);
+								GenerateReport(cachedCsvFile, outputDir, bBulkMode, rowData, bBatchedGraphs, writeDetailedReports, bReadAllStats || bWriteToSummaryTableCache, cachedCsvFile.reportTypeInfo, csvDir);
 								perfLog.LogTiming("  GenerateReport");
 								if (rowData != null && bWriteToSummaryTableCache)
 								{
@@ -1086,7 +1087,7 @@ namespace PerfReportTool
 			return csvStats;
 		}
 
-		void GenerateReport(CachedCsvFile csvFile, string outputDir, bool bBulkMode, SummaryTableRowData rowData, bool bBatchedGraphs, bool writeDetailedReport, bool bReadAllStats, ReportTypeInfo reportTypeInfo)
+		void GenerateReport(CachedCsvFile csvFile, string outputDir, bool bBulkMode, SummaryTableRowData rowData, bool bBatchedGraphs, bool writeDetailedReport, bool bReadAllStats, ReportTypeInfo reportTypeInfo, string csvDir)
         {
             PerfLog perfLog = new PerfLog(GetBoolArg("perfLog"));
             string shortName = ReplaceFileExtension(MakeShortFilename(csvFile.filename), "");
@@ -1220,11 +1221,18 @@ namespace PerfReportTool
 				Uri relativeCsvUri = finalDirUri.MakeRelativeUri(csvFileUri);
 				string csvPath = relativeCsvUri.ToString();
 
+				// re-root the CSV path if requested
 				string csvLinkRootPath = GetArg("csvLinkRootPath", null);
-				if ( csvLinkRootPath != null && !relativeCsvUri.IsAbsoluteUri)
+				if ( csvDir != null && csvLinkRootPath != null)
 				{
-					csvPath = csvLinkRootPath + csvPath;
+					string csvDirFinal = csvDir.Replace("\\", "/");
+					csvDirFinal += csvDirFinal.EndsWith("/") ? "" : "/";
+					Uri csvDirUri = new Uri(csvDirFinal, UriKind.Absolute);
+					Uri csvRelativeToCsvDirUri = csvDirUri.MakeRelativeUri(csvFileUri);
+					csvPath = Path.Combine(csvLinkRootPath,csvRelativeToCsvDirUri.ToString());
+					csvPath = new Uri(csvPath, UriKind.Absolute).ToString();
 				}
+
 				rowData.Add(SummaryTableElement.Type.ToolMetadata, "Csv File", "<a href='" + csvPath + "'>" + shortName + ".csv" + "</a>", null, csvPath);
 				rowData.Add(SummaryTableElement.Type.ToolMetadata, "ReportType", reportTypeInfo.name);
 				rowData.Add(SummaryTableElement.Type.ToolMetadata, "ReportTypeID", reportTypeInfo.summaryTableCacheID);
