Integrate CL 16308862 from Dev-EngineMerge: PerfReportTool 4.47 - ExtraLinks summary
Change 16308862 by ben.woodhouse[REDACTED_HANDLE] on 2021/05/12 22:00:30

PerfReportTool 4.47 - ExtraLinks summary
Adds a new ExtraLinks summary type, which reads a <csvfilename>.links file. The file contains lines with additional links in the format:
LongName,ShortName,URL

These links are added to a new Links section of the report. They are also added to summary table row, with the LongName as the key.

The section can be added to all reports by adding with -extraLinksSummary, or add the summary type "ExtraLinks" to the Report type XML.

diff --git a/Engine/Source/Programs/CSVTools/PerfReportTool/PerfReportTool.cs b/Engine/Source/Programs/CSVTools/PerfReportTool/PerfReportTool.cs
--- a/Engine/Source/Programs/CSVTools/PerfReportTool/PerfReportTool.cs
+++ b/Engine/Source/Programs/CSVTools/PerfReportTool/PerfReportTool.cs
@@ -22,7 +22,7 @@ namespace PerfReportTool
 {
     class Version
     {
-        private static string VersionString = "4.46";
+        private static string VersionString = "4.47";
 
         public static string Get() { return VersionString; }
     };
@@ -1471,9 +1471,16 @@ namespace PerfReportTool
 
 			bool bIncludeSummaryCsv = GetBoolArg("writeSummaryCsv") && !bBulkMode;
 
-            // If the reporttype has summary info, then write out the summary]
-            PeakSummary peakSummary = null;
-            foreach (Summary summary in reportTypeInfo.summaries)
+			List<Summary> summaries = new List<Summary>(reportTypeInfo.summaries);
+			bool bExtraLinksSummary = GetBoolArg("extraLinksSummary");
+			if (bExtraLinksSummary)
+			{
+				summaries.Insert(0,new ExtraLinksSummary(null, null));
+			}
+
+			// If the reporttype has summary info, then write out the summary]
+			PeakSummary peakSummary = null;
+            foreach (Summary summary in summaries)
             {
                 summary.WriteSummaryData(htmlFile, summary.useUnstrippedCsvStats ? unstrippedCsvStats : csvStats, bIncludeSummaryCsv, summaryRowData, htmlFilename);
                 if ( summary.GetType() == typeof(PeakSummary) )
@@ -1895,34 +1902,7 @@ namespace PerfReportTool
 						summaryElement = child;
 					}
 					string summaryType = summaryElement.Attribute("type").Value;
-					if (summaryType == "histogram")
-					{
-						summaries.Add(new HistogramSummary(summaryElement, baseXmlDirectory));
-					}
-					else if (summaryType == "peak")
-					{
-						summaries.Add(new PeakSummary(summaryElement, baseXmlDirectory));
-					}
-					else if (summaryType == "fpschart")
-					{
-						summaries.Add(new FPSChartSummary(summaryElement, baseXmlDirectory));
-					}
-					else if (summaryType == "hitches")
-					{
-						summaries.Add(new HitchSummary(summaryElement, baseXmlDirectory));
-					}
-					else if (summaryType == "event")
-					{
-						summaries.Add(new EventSummary(summaryElement, baseXmlDirectory));
-					}
-					else if (summaryType == "boundedstatvalues")
-					{
-						summaries.Add(new BoundedStatValuesSummary(summaryElement, baseXmlDirectory));
-					}
-					else if (summaryType == "mapoverlay")
-					{
-						summaries.Add(new MapOverlaySummary(summaryElement, baseXmlDirectory));
-					}
+					summaries.Add(SummaryFactory.Create(summaryType, summaryElement, baseXmlDirectory));
 				}
 				else if (child.Name == "metadataToShow")
 				{
diff --git a/Engine/Source/Programs/CSVTools/PerfReportTool/perfsummaries.cs b/Engine/Source/Programs/CSVTools/PerfReportTool/perfsummaries.cs
--- a/Engine/Source/Programs/CSVTools/PerfReportTool/perfsummaries.cs
+++ b/Engine/Source/Programs/CSVTools/PerfReportTool/perfsummaries.cs
@@ -10,6 +10,48 @@ using PerfReportTool;
 
 namespace PerfSummaries
 {
+	class SummaryFactory
+	{
+		public static Summary Create(string summaryType, XElement summaryXmlElement, string baseXmlDirectory)
+		{
+			switch (summaryType)
+			{
+				case "histogram":
+				{
+					return new HistogramSummary(summaryXmlElement, baseXmlDirectory);
+				}
+				case "peak":
+				{
+					return new PeakSummary(summaryXmlElement, baseXmlDirectory);
+				}
+				case "fpschart":
+				{
+					return new FPSChartSummary(summaryXmlElement, baseXmlDirectory);
+				}
+				case "hitches":
+				{
+					return new HitchSummary(summaryXmlElement, baseXmlDirectory);
+				}
+				case "event":
+				{
+					return new EventSummary(summaryXmlElement, baseXmlDirectory);
+				}
+				case "boundedstatvalues":
+				{
+					return new BoundedStatValuesSummary(summaryXmlElement, baseXmlDirectory);
+				}
+				case "mapoverlay":
+				{
+					return new MapOverlaySummary(summaryXmlElement, baseXmlDirectory);
+				}
+				case "extralinks":
+				{
+					return new ExtraLinksSummary(summaryXmlElement, baseXmlDirectory);
+				}
+			}
+			throw new Exception("Summary type "+summaryType+" not found!");
+		}
+	};
 
     class Colour
     {
@@ -593,7 +635,7 @@ namespace PerfSummaries
 				ColumnColors.Add(ColourThresholdList.GetSafeColourForValue(ColumnColorThresholds.Last(), ColumnValues.Last()));
 			}
 
-			// Output row data
+			// Output summary table row data
 			if (rowData != null)
             {
                 for (int i = 0; i < ColumnNames.Count; i++)
@@ -827,8 +869,8 @@ namespace PerfSummaries
                 htmlFile.WriteLine("  </table>");
             }
 
-            // Output metadata
-            if (rowData != null)
+			// Output summary table row data
+			if (rowData != null)
             {
 
                 ColourThresholdList thresholdList = null;
@@ -869,7 +911,7 @@ namespace PerfSummaries
 
         public override void WriteSummaryData(System.IO.StreamWriter htmlFile, CsvStats csvStats, bool bIncludeSummaryCsv, SummaryTableRowData metadata, string htmlFileName)
         {
-			// Only HTML reporting is supported (does not output metadata)
+			// Only HTML reporting is supported (does not summary table row data)
 			if (htmlFile == null)
 			{
 				return;
@@ -977,7 +1019,7 @@ namespace PerfSummaries
 
         public override void WriteSummaryData(System.IO.StreamWriter htmlFile, CsvStats csvStats, bool bIncludeSummaryCsv, SummaryTableRowData metadata, string htmlFileName)
         {
-			// Only HTML reporting is supported (does not output metadata)
+			// Only HTML reporting is supported (does not output summary table row data)
 			if (htmlFile == null)
 			{
 				return;
@@ -1217,7 +1259,7 @@ namespace PerfSummaries
 
         public override void WriteSummaryData(System.IO.StreamWriter htmlFile, CsvStats csvStats, bool bIncludeSummaryCsv, SummaryTableRowData metadata, string htmlFileName)
         {
-			// Only HTML reporting is supported (does not output metadata)
+			// Only HTML reporting is supported (does not output summary table row data)
 			if (htmlFile == null)
 			{
 				return;
@@ -1627,7 +1669,7 @@ namespace PerfSummaries
 				htmlFile.WriteLine("  </table>");
 			}
 
-			// Output metadata
+			// Output summary table row data
 			if (metadata != null)
 			{
 				foreach (Column col in filteredColumns)
@@ -1709,6 +1751,10 @@ namespace PerfSummaries
 			}
 
 			sourceImagePath = element.GetSafeAttibute<string>("sourceImage");
+			if (baseXmlDirectory == null)
+			{
+				throw new Exception("BaseXmlDirectory not specified");
+			}
 			if ( !System.IO.Path.IsPathRooted(sourceImagePath))
 			{
 				sourceImagePath = System.IO.Path.GetFullPath(System.IO.Path.Combine(baseXmlDirectory,sourceImagePath));
@@ -1905,6 +1951,94 @@ namespace PerfSummaries
 		List<MapOverlay> overlays = new List<MapOverlay>();
 	};
 
+	class ExtraLinksSummary : Summary
+	{
+		class ExtraLink
+		{
+			public ExtraLink(string fileLine)
+			{
+				string[] Sections = fileLine.Split(',');
+				if ( Sections.Length != 3 )
+				{
+					throw new Exception("Bad links line format: "+fileLine);
+				}
+				LongName = Sections[0];
+				ShortName = Sections[1];
+				LinkURL = Sections[2];
+			}
+			public string GetLinkString(bool bUseLongName)
+			{
+				string Text = bUseLongName ? LongName : ShortName;
+				return "<a href='" + LinkURL + "'>"+ Text + "</a>";
+			}
+			public string LongName;
+			public string ShortName;
+			public string LinkURL;
+		};
+
+		public ExtraLinksSummary(XElement element, string baseXmlDirectory)
+		{
+			title = "Links";
+			if (element != null)
+			{
+				title = element.GetSafeAttibute("title", title);
+			}
+		}
+
+		public override void WriteSummaryData(System.IO.StreamWriter htmlFile, CsvStats csvStats, bool bIncludeSummaryCsv, SummaryTableRowData rowData, string htmlFileName)
+		{
+			List<ExtraLink> links = new List<ExtraLink>();
+
+			string csvFilename = csvStats.metaData.GetValue("csvfilename", null);
+			if (csvFilename == null)
+			{
+				Console.WriteLine("Can't find CSV filename for ExtraLinks summary. Skipping");
+				return;
+			}
+
+			string linksFilename = csvFilename + ".links";
+			if (!File.Exists(linksFilename))
+			{
+				Console.WriteLine("Can't find file " + linksFilename + " for ExtraLinks summary. Skipping");
+				return;
+			}
+			string[] lines = File.ReadAllLines(linksFilename);
+			foreach (string line in lines)
+			{
+				links.Add(new ExtraLink(line));
+			}
+			if (links.Count == 0)
+			{
+				return;
+			}
+
+			// Output HTML
+			if (htmlFile != null)
+			{
+				htmlFile.WriteLine("  <h2>" + title + "</h2>");
+				htmlFile.WriteLine("  <ul>");
+				foreach (ExtraLink link in links)
+				{
+					htmlFile.WriteLine("  <li>" + link.GetLinkString(true) + "</li>");
+				}
+				htmlFile.WriteLine("  </ul>");
+			}
+
+			// Output summary row data
+			if (rowData != null)
+			{
+				foreach (ExtraLink link in links)
+				{
+					rowData.Add(SummaryTableElement.Type.SummaryTableMetric, link.LongName, link.GetLinkString(false), null);
+				}
+			}
+		}
+		public override void PostInit(ReportTypeInfo reportTypeInfo, CsvStats csvStats)
+		{
+		}
+		string title;
+	};
+
 	class SummaryTableElement
     {
 		// Bump this when making changes!
