PerfReportTool 4.49:  Fix latent bug with PRC color serialization (ToString() not implemented).
- Don't specify colours for PeakSummary report type. Pass null instead and have it set them automatically (this also avoids the need to serialize)

diff --git a/Engine/Source/Programs/CSVTools/PerfReportTool/PerfReportTool.cs b/Engine/Source/Programs/CSVTools/PerfReportTool/PerfReportTool.cs
--- a/Engine/Source/Programs/CSVTools/PerfReportTool/PerfReportTool.cs
+++ b/Engine/Source/Programs/CSVTools/PerfReportTool/PerfReportTool.cs
@@ -22,7 +22,7 @@ namespace PerfReportTool
 {
     class Version
     {
-        private static string VersionString = "4.48";
+        private static string VersionString = "4.49";
 
         public static string Get() { return VersionString; }
     };
diff --git a/Engine/Source/Programs/CSVTools/PerfReportTool/perfsummaries.cs b/Engine/Source/Programs/CSVTools/PerfReportTool/perfsummaries.cs
--- a/Engine/Source/Programs/CSVTools/PerfReportTool/perfsummaries.cs
+++ b/Engine/Source/Programs/CSVTools/PerfReportTool/perfsummaries.cs
@@ -94,15 +94,18 @@ namespace PerfSummaries
 
         public string ToHTMLString()
         {
-            int rI = (int)(r * 255.0f);
-            int gI = (int)(g * 255.0f);
-            int bI = (int)(b * 255.0f);
-            int aI = (int)(alpha * 255.0f);
-            return "'#" + rI.ToString("x2") + gI.ToString("x2") + bI.ToString("x2") /*+ aI.ToString("X")*/ + "'";
+            return "'" + ToString() + "'";
         }
 
+		public override string ToString() 
+		{
+			int rI = (int)(r * 255.0f);
+			int gI = (int)(g * 255.0f);
+			int bI = (int)(b * 255.0f);
+			return "#" + rI.ToString("x2") + gI.ToString("x2") + bI.ToString("x2");
+		}
 
-        public static Colour White = new Colour(1.0f, 1.0f, 1.0f, 1.0f);
+		public static Colour White = new Colour(1.0f, 1.0f, 1.0f, 1.0f);
         public static Colour Black = new Colour(0, 0, 0, 1.0f);
         public static Colour Orange = new Colour(1.0f, 0.5f, 0.0f, 1.0f);
 		public static Colour Yellow = new Colour(1.0f, 1.0f, 0.0f, 1.0f);
@@ -116,10 +119,10 @@ namespace PerfSummaries
 
 	class ThresholdInfo
 	{
-		public ThresholdInfo(double inValue, Colour inColour)
+		public ThresholdInfo(double inValue, Colour inColourOverride=null)
 		{
 			value = inValue;
-			colour = inColour;
+			colour = inColourOverride;
 		}
 		public double value;
 		public Colour colour;
@@ -889,7 +892,7 @@ namespace PerfSummaries
                     thresholdList = new ColourThresholdList();
                     for (int i = 0; i < colourThresholds.Length; i++)
                     {
-                        thresholdList.Add(new ThresholdInfo(colourThresholds[i], null));
+                        thresholdList.Add(new ThresholdInfo(colourThresholds[i]));
                     }
                 }
                 rowData.Add(SummaryTableElement.Type.SummaryTableMetric, summaryStatName, (double)eventCount, thresholdList);
@@ -1250,10 +1253,10 @@ namespace PerfSummaries
 					float yellowValue = (float)budget * 1.0f;
 					float greenValue = (float)budget * 0.9f;
 
-					colorThresholdList.Add(new ThresholdInfo(redValue, Colour.Red));
-					colorThresholdList.Add(new ThresholdInfo(orangeValue, Colour.Orange));
-					colorThresholdList.Add(new ThresholdInfo(yellowValue, Colour.Yellow));
-					colorThresholdList.Add(new ThresholdInfo(greenValue, Colour.Green));
+					colorThresholdList.Add(new ThresholdInfo(greenValue));
+					colorThresholdList.Add(new ThresholdInfo(yellowValue));
+					colorThresholdList.Add(new ThresholdInfo(orangeValue));
+					colorThresholdList.Add(new ThresholdInfo(redValue));
 
 					peakColour = colorThresholdList.GetColourForValue(peak);
 					averageColour = colorThresholdList.GetColourForValue(average);
