PerfReportTool 4.49.1: Fix access none with single CSV mode peak summary

diff --git a/Engine/Source/Programs/CSVTools/PerfReportTool/PerfReportTool.cs b/Engine/Source/Programs/CSVTools/PerfReportTool/PerfReportTool.cs
--- a/Engine/Source/Programs/CSVTools/PerfReportTool/PerfReportTool.cs
+++ b/Engine/Source/Programs/CSVTools/PerfReportTool/PerfReportTool.cs
@@ -22,7 +22,7 @@ namespace PerfReportTool
 {
     class Version
     {
-        private static string VersionString = "4.49";
+        private static string VersionString = "4.49.1";
 
         public static string Get() { return VersionString; }
     };
diff --git a/Engine/Source/Programs/CSVTools/PerfReportTool/perfsummaries.cs b/Engine/Source/Programs/CSVTools/PerfReportTool/perfsummaries.cs
--- a/Engine/Source/Programs/CSVTools/PerfReportTool/perfsummaries.cs
+++ b/Engine/Source/Programs/CSVTools/PerfReportTool/perfsummaries.cs
@@ -1272,16 +1272,20 @@ namespace PerfSummaries
 					LLMCsvData.WriteLine(string.Format("{0}_Colors,{1},{2},'#aaaaaa'", csvStatName, averageColour, peakColour));
 				}
 
-				SummaryTableElement smv;
-				if (summaryTableRowData.dict.TryGetValue(statName.ToLower(), out smv))
+				if (summaryTableRowData != null)
 				{
-					if (smv.type == SummaryTableElement.Type.CsvStatAverage)
+					SummaryTableElement smv;
+					// Hide duplicate CsvStatAverage stats
+					if (summaryTableRowData.dict.TryGetValue(statName.ToLower(), out smv))
 					{
-						smv.SetFlag(SummaryTableElement.Flags.Hidden, true);
+						if (smv.type == SummaryTableElement.Type.CsvStatAverage)
+						{
+							smv.SetFlag(SummaryTableElement.Flags.Hidden, true);
+						}
 					}
+					summaryTableRowData.Add(SummaryTableElement.Type.SummaryTableMetric, statInfo.shortName + " Avg", average, colorThresholdList);
+					summaryTableRowData.Add(SummaryTableElement.Type.SummaryTableMetric, statInfo.shortName + " Max", peak, colorThresholdList);
 				}
-				summaryTableRowData.Add(SummaryTableElement.Type.SummaryTableMetric, statInfo.shortName + " Avg", average, colorThresholdList);
-				summaryTableRowData.Add(SummaryTableElement.Type.SummaryTableMetric, statInfo.shortName + " Max", peak, colorThresholdList);
 
 			}
 			htmlFile.WriteLine("  </table>");
