PerfReportTool 4.49.2 - Fix non-rounded numbers in the FPSChart summary
Update test data

diff --git a/Engine/Source/Programs/CSVTools/PerfReportTool/PerfReportTool.cs b/Engine/Source/Programs/CSVTools/PerfReportTool/PerfReportTool.cs
--- a/Engine/Source/Programs/CSVTools/PerfReportTool/PerfReportTool.cs
+++ b/Engine/Source/Programs/CSVTools/PerfReportTool/PerfReportTool.cs
@@ -22,7 +22,7 @@ namespace PerfReportTool
 {
     class Version
     {
-        private static string VersionString = "4.49.1";
+        private static string VersionString = "4.49.2";
 
         public static string Get() { return VersionString; }
     };
diff --git a/Engine/Source/Programs/CSVTools/PerfReportTool/perfsummaries.cs b/Engine/Source/Programs/CSVTools/PerfReportTool/perfsummaries.cs
--- a/Engine/Source/Programs/CSVTools/PerfReportTool/perfsummaries.cs
+++ b/Engine/Source/Programs/CSVTools/PerfReportTool/perfsummaries.cs
@@ -692,7 +692,7 @@ namespace PerfSummaries
                         columnName += " (ms)";
                     }
                     HeaderRow += "<th>" + TableUtil.FormatStatName(columnName) + "</th>";
-                    ValueRow += "<td bgcolor=" + ColumnColors[i] + ">" + ColumnValues[i] + "</td>";
+                    ValueRow += "<td bgcolor=" + ColumnColors[i] + ">" + ColumnValues[i].ToString("0.00") + "</td>";
                 }
 				htmlFile.WriteLine("  <h2>FPSChart</h2>");
 				htmlFile.WriteLine("<table border='0' style='width:400'>");
@@ -787,7 +787,7 @@ namespace PerfSummaries
                         ValueRow += "<td>"+ CapRange.name + "</td>";
                         for (int i = 0; i < ColumnNames.Count; i++)
                         {
-                            ValueRow += "<td bgcolor=" + ColumnColors[i] + ">" + ColumnValues[i] + "</td>";
+                            ValueRow += "<td bgcolor=" + ColumnColors[i] + ">" + ColumnValues[i].ToString("0.00") + "</td>";
                         }
                         htmlFile.WriteLine("  <tr>" + ValueRow + "</tr>");
 					}
