PerfReportTool 4.51 - scrollable summary table improvements
- Add support for more than two sticky columns in collated summary tables
- Add format info for columns, replaces "lowIsBad" list. Allows specifying auto colorization rules and also numerical formatting
- Add support for multiple section boundaries, and minor (dashed line) section boundaries
- Section boundaries no longer require startToken or endToken (both are now optional)
- Section boundaries now can be specific to collated or full tables

diff --git a/Engine/Source/Programs/CSVTools/PerfReportTool/PerfReportTool.cs b/Engine/Source/Programs/CSVTools/PerfReportTool/PerfReportTool.cs
--- a/Engine/Source/Programs/CSVTools/PerfReportTool/PerfReportTool.cs
+++ b/Engine/Source/Programs/CSVTools/PerfReportTool/PerfReportTool.cs
@@ -22,7 +22,7 @@ namespace PerfReportTool
 {
     class Version
     {
-        private static string VersionString = "4.50";
+        private static string VersionString = "4.51";
 
         public static string Get() { return VersionString; }
     };
@@ -962,10 +962,10 @@ namespace PerfReportTool
         void WriteSummaryTableReport(string outputDir, string filenameWithoutExtension, SummaryTable table, SummaryTableInfo tableInfo, bool bCollated, bool bToCSV, bool bSpreadsheetFriendlyStrings, string weightByColumnNameOverride)
 		{
 			string weightByColumnName = weightByColumnNameOverride != null ? weightByColumnNameOverride : tableInfo.weightByColumn;
-			WriteSummaryTableReport(outputDir, filenameWithoutExtension, table, tableInfo.columnFilterList, tableInfo.rowSortList, bCollated, bToCSV, bSpreadsheetFriendlyStrings, tableInfo.sectionBoundary, weightByColumnName);
+			WriteSummaryTableReport(outputDir, filenameWithoutExtension, table, tableInfo.columnFilterList, tableInfo.rowSortList, bCollated, bToCSV, bSpreadsheetFriendlyStrings, tableInfo.sectionBoundaries, weightByColumnName);
 		}
 
-		void WriteSummaryTableReport(string outputDir, string filenameWithoutExtension, SummaryTable table, List<string> columnFilterList, List<string> rowSortList, bool bCollated, bool bToCSV, bool bSpreadsheetFriendlyStrings, SummarySectionBoundaryInfo sectionBoundaryInfo, string weightByColumnName)
+		void WriteSummaryTableReport(string outputDir, string filenameWithoutExtension, SummaryTable table, List<string> columnFilterList, List<string> rowSortList, bool bCollated, bool bToCSV, bool bSpreadsheetFriendlyStrings, List<SummarySectionBoundaryInfo> sectionBoundaries, string weightByColumnName)
 		{
 			if (GetBoolArg("noWeightedAvg"))
 			{
@@ -991,7 +991,7 @@ namespace PerfReportTool
 			{
 				filteredTable.ApplyDisplayNameMapping(statDisplaynameMapping);
 				string VersionString = GetBoolArg("noWatermarks") ? "" : Version.Get();
-				filteredTable.WriteToHTML(filenameWithoutExtension+".html", VersionString, bSpreadsheetFriendlyStrings, sectionBoundaryInfo, bScrollableTable, addMinMaxColumns, GetIntArg("maxSummaryTableStringLength", -1), reportXML.summaryTableLowIsBadStatList, weightByColumnName);
+				filteredTable.WriteToHTML(filenameWithoutExtension+".html", VersionString, bSpreadsheetFriendlyStrings, sectionBoundaries, bScrollableTable, addMinMaxColumns, GetIntArg("maxSummaryTableStringLength", -1), reportXML.columnFormatInfoList, weightByColumnName);
 			}
 		}
 
@@ -2297,10 +2297,10 @@ namespace PerfReportTool
 				}
 			}
 
-			XElement summaryTableLowIsBadStatListEl = rootElement.Element("summaryTableLowIsBadStats");
-			if (summaryTableLowIsBadStatListEl != null)
+			XElement summaryTableColumnInfoListEl = rootElement.Element("summaryTableColumnFormatInfo");
+			if (summaryTableColumnInfoListEl != null)
 			{
-				summaryTableLowIsBadStatList = summaryTableLowIsBadStatListEl.Value.Split(',');
+				columnFormatInfoList = new SummaryTableColumnFormatInfoCollection(summaryTableColumnInfoListEl);
 			}			
 
 			// Read the derived metadata mappings
@@ -2543,7 +2543,7 @@ namespace PerfReportTool
 		Dictionary<string,XElement> sharedSummaries;
 		Dictionary<string, GraphSettings> graphs;
 		Dictionary<string, string> statDisplayNameMapping;
-		public string [] summaryTableLowIsBadStatList;
+		public SummaryTableColumnFormatInfoCollection columnFormatInfoList;
 		string baseXmlDirectory;
 
 		List<CsvEventStripInfo> csvEventsToStrip;
diff --git a/Engine/Source/Programs/CSVTools/PerfReportTool/perfsummaries.cs b/Engine/Source/Programs/CSVTools/PerfReportTool/perfsummaries.cs
--- a/Engine/Source/Programs/CSVTools/PerfReportTool/perfsummaries.cs
+++ b/Engine/Source/Programs/CSVTools/PerfReportTool/perfsummaries.cs
@@ -2438,10 +2438,23 @@ namespace PerfSummaries
 
 	class SummarySectionBoundaryInfo
 	{
+		public SummarySectionBoundaryInfo( string inStatName, string inStartToken, string inEndToken, bool inIsMinor, bool inInCollatedTable, bool inInFullTable)
+		{
+			statName = inStatName;
+			startToken = inStartToken;
+			endToken = inEndToken;
+			isMinor = inIsMinor;
+			inCollatedTable = inInCollatedTable;
+			inFullTable = inInFullTable;
+		}
+		public string statName;
 		public string startToken;
 		public string endToken;
-		public string statName;
+		public bool isMinor;
+		public bool inCollatedTable;
+		public bool inFullTable;
 	};
+
 	class SummaryTableInfo
 	{
 		public SummaryTableInfo(XElement tableElement)
@@ -2462,13 +2475,21 @@ namespace PerfSummaries
 			{
 				columnFilterList.AddRange(filterEl.Value.Split(','));
 			}
-			XElement sectionBoundaryEl = tableElement.Element("sectionBoundary");
-			if (sectionBoundaryEl != null)
+
+			foreach (XElement sectionBoundaryEl in tableElement.Elements("sectionBoundary"))
 			{
-				sectionBoundary = new SummarySectionBoundaryInfo();
-				sectionBoundary.statName = sectionBoundaryEl.Attribute("statName").Value;
-				sectionBoundary.startToken = sectionBoundaryEl.Attribute("startToken").Value;
-				sectionBoundary.endToken = sectionBoundaryEl.Attribute("endToken").Value;
+				if (sectionBoundaryEl != null)
+				{
+					SummarySectionBoundaryInfo sectionBoundary = new SummarySectionBoundaryInfo(
+						sectionBoundaryEl.GetSafeAttibute<string>("statName"),
+						sectionBoundaryEl.GetSafeAttibute<string>("startToken"),
+						sectionBoundaryEl.GetSafeAttibute<string>("endToken"),
+						sectionBoundaryEl.GetSafeAttibute<bool>("minor", false),
+						sectionBoundaryEl.GetSafeAttibute<bool>("inCollatedTable", true),
+						sectionBoundaryEl.GetSafeAttibute<bool>("inFullTable", true)
+						);
+					sectionBoundaries.Add(sectionBoundary);
+				}
 			}
 		}
 
@@ -2480,11 +2501,93 @@ namespace PerfSummaries
 
 		public List<string> rowSortList = new List<string>();
 		public List<string> columnFilterList = new List<string>();
-		public SummarySectionBoundaryInfo sectionBoundary = null;
+		public List<SummarySectionBoundaryInfo> sectionBoundaries = new List<SummarySectionBoundaryInfo>();
 		public string weightByColumn = null;
 	}
 
 
+	class SummaryTableColumnFormatInfoCollection
+	{
+		public SummaryTableColumnFormatInfoCollection(XElement element)
+		{
+			foreach (XElement child in element.Elements("columnInfo"))
+			{
+				columnFormatInfoList.Add(new SummaryTableColumnFormatInfo(child));
+			}
+		}
+
+		public SummaryTableColumnFormatInfo GetFormatInfo(string columnName)
+		{
+			string lowerColumnName = columnName.ToLower();
+			if (lowerColumnName.StartsWith("avg ") || lowerColumnName.StartsWith("min ") || lowerColumnName.StartsWith("max "))
+			{
+				lowerColumnName = lowerColumnName.Substring(4);
+			}
+			foreach (SummaryTableColumnFormatInfo columnInfo in columnFormatInfoList)
+			{
+				int wildcardIndex = columnInfo.name.IndexOf('*');
+				if (wildcardIndex == -1)
+				{
+					if (columnInfo.name == lowerColumnName)
+					{
+						return columnInfo;
+					}
+				}
+				else
+				{
+					string prefix = columnInfo.name.Substring(0, wildcardIndex);
+					if (lowerColumnName.StartsWith(prefix))
+					{
+						return columnInfo;
+					}
+				}
+			}
+			return defaultColumnInfo;
+		}
+
+		public static SummaryTableColumnFormatInfo DefaultColumnInfo 
+		{
+			get { return defaultColumnInfo; }
+		}
+
+		List<SummaryTableColumnFormatInfo> columnFormatInfoList = new List<SummaryTableColumnFormatInfo>();
+		static SummaryTableColumnFormatInfo defaultColumnInfo = new SummaryTableColumnFormatInfo();
+	};
+
+	enum AutoColorizeMode
+	{
+		Off,
+		HighIsBad,
+		LowIsBad,
+	};
+
+	class SummaryTableColumnFormatInfo
+	{
+		public SummaryTableColumnFormatInfo()
+		{
+			name = "Default";
+		}
+		public SummaryTableColumnFormatInfo(XElement element)
+		{
+			name = element.Attribute("name").Value.ToLower();
+
+			string autoColorizeStr =element.GetSafeAttibute<string>("autoColorize", "highIsBad").ToLower();
+			var modeList=Enum.GetValues(typeof(AutoColorizeMode));
+			foreach (AutoColorizeMode mode in modeList)
+			{
+				if (mode.ToString().ToLower() == autoColorizeStr)
+				{
+					autoColorizeMode = mode;
+					break;
+				}
+			}
+			numericFormat = element.GetSafeAttibute<string>("numericFormat");
+		}
+		public AutoColorizeMode autoColorizeMode = AutoColorizeMode.HighIsBad;
+		public string name;
+		public string numericFormat;
+	};
+
 	class SummaryTableColumn
 	{
 		public string name;
@@ -2604,8 +2707,12 @@ namespace PerfSummaries
 			return thresholds.GetColourForValue(value);
 		}
 
-		public void ComputeAutomaticColourThresholds(bool bHighIsRed)
+		public void ComputeAutomaticColourThresholds(AutoColorizeMode autoColorizeMode)
 		{
+			if ( autoColorizeMode == AutoColorizeMode.Off)
+			{
+				return;
+			}
 			colourThresholds = new List<ColourThresholdList>();
 			double maxValue = -double.MaxValue;
 			double minValue = double.MaxValue;
@@ -2633,10 +2740,10 @@ namespace PerfSummaries
 
 			double averageValue = totalValue / validCount; // TODO: Weighted average 
 			colourThresholdOverride = new ColourThresholdList();
-			colourThresholdOverride.Add(new ThresholdInfo(minValue, bHighIsRed ? green : red));
+			colourThresholdOverride.Add(new ThresholdInfo(minValue, (autoColorizeMode == AutoColorizeMode.HighIsBad) ? green : red));
 			colourThresholdOverride.Add(new ThresholdInfo(averageValue, yellow));
 			colourThresholdOverride.Add(new ThresholdInfo(averageValue, yellow));
-			colourThresholdOverride.Add(new ThresholdInfo(maxValue, bHighIsRed ? red : green));
+			colourThresholdOverride.Add(new ThresholdInfo(maxValue, (autoColorizeMode == AutoColorizeMode.HighIsBad) ? red : green));
 		}
 
 		public int GetCount()
@@ -2670,7 +2777,7 @@ namespace PerfSummaries
 			stringValues[index] = value;
 			isNumeric = false;
 		}
-		public string GetStringValue(int index, bool roundNumericValues = false)
+		public string GetStringValue(int index, bool roundNumericValues = false, string forceNumericFormat = null)
 		{
 			if (isNumeric)
 			{
@@ -2679,7 +2786,11 @@ namespace PerfSummaries
 					return "";
 				}
 				double val = doubleValues[index];
-				if (roundNumericValues)
+				if (forceNumericFormat != null)
+				{
+					return val.ToString(forceNumericFormat);
+				}
+				else if (roundNumericValues)
 				{
 					double absVal = Math.Abs(val);
 					double frac = absVal - (double)Math.Truncate(absVal);
@@ -2705,6 +2816,16 @@ namespace PerfSummaries
 				{
 					return "";
 				}
+				if (forceNumericFormat != null)
+				{
+					// We're forcing a numeric format on something that's technically a string, but since we were asked, we'll try to do it anyway 
+					// Note: this is not ideal, but it's useful for collated table columns, which might get converted to non-numeric during collation
+					try
+					{
+						return Convert.ToDouble(stringValues[index], System.Globalization.CultureInfo.InvariantCulture).ToString(forceNumericFormat);
+					}
+					catch { } // Ignore. Just fall through...
+				}
 				return stringValues[index];
 			}
 		}
@@ -3064,50 +3185,8 @@ namespace PerfSummaries
 			csvFile.Close();
 		}
 
-		private void AutoColorizeColumns(string[] summaryTableLowIsBadStatList)
-		{
-			foreach (SummaryTableColumn column in columns)
-			{
-				if (column.isNumeric)
-				{
-					bool highIsRed = true;
-					if (summaryTableLowIsBadStatList != null)
-					{
-						// Determine if this is a low is red column
-						string lowerColumnName = column.name.ToLower();
-						if (lowerColumnName.StartsWith("avg ") || lowerColumnName.StartsWith("min ") || lowerColumnName.StartsWith("max "))
-						{
-							lowerColumnName = lowerColumnName.Substring(4);
-						}
-						foreach (string entry in summaryTableLowIsBadStatList)
-						{
-							string lowerEntry = entry.ToLower();
-							int wildcardIndex = lowerEntry.IndexOf('*');
-							if (wildcardIndex == -1)
-							{
-								if (lowerEntry == lowerColumnName)
-								{
-									highIsRed = false;
-									break;
-								}
-							}
-							else
-							{
-								string prefix = lowerEntry.Substring(0, wildcardIndex);
-								if (lowerColumnName.StartsWith(prefix))
-								{
-									highIsRed = false;
-									break;
-								}
-							}
-						}
-					}
-					column.ComputeAutomaticColourThresholds(highIsRed);
-				}
-			}
-		}
 
-		public void WriteToHTML(string htmlFilename, string VersionString, bool bSpreadsheetFriendlyStrings, SummarySectionBoundaryInfo sectionBoundaryInfo, bool bScrollableTable, bool bAddMinMaxColumns, int maxColumnStringLength, string [] summaryTableLowIsBadStatList, string weightByColumnName)
+		public void WriteToHTML(string htmlFilename, string VersionString, bool bSpreadsheetFriendlyStrings, List<SummarySectionBoundaryInfo> sectionBoundaries, bool bScrollableTable, bool bAddMinMaxColumns, int maxColumnStringLength, SummaryTableColumnFormatInfoCollection columnFormatInfoList, string weightByColumnName)
 		{
 			System.IO.StreamWriter htmlFile = new System.IO.StreamWriter(htmlFilename, false);
 			int statColSpan = hasMinMaxColumns ? 3 : 1;
@@ -3119,14 +3198,77 @@ namespace PerfSummaries
 
 			htmlFile.WriteLine("<html>");
 			htmlFile.WriteLine("<head><title>Performance summary</title>");
-			//htmlFile.WriteLine("table, th, td { border: 0px solid black; border-collapse: separate; padding: 3px; vertical-align: top; font-family: 'Verdana', Times, serif; font-size: 12px;}");
-			htmlFile.WriteLine("<style type='text/css'>");
+
+			// Figure out the sticky column count
+			int stickyColumnCount = 0;
+			if (bScrollableTable)
+			{
+				stickyColumnCount = 1;
+				if (isCollated)
+				{
+					for (int i=0; i<columns.Count; i++)
+					{
+						if ( columns[i].name == "Count" )
+						{
+							stickyColumnCount = i + 1;
+							break;
+						}
+					}
+				}
+			}
+
+			if (bScrollableTable)
+			{
+				// Insert some javascript to make the columns sticky. It's not possible to do this for multiple columns with pure CSS, since you need to compute the X offset dynamically
+				// We need to do this when the page is loaded or the window is resized
+				htmlFile.WriteLine("<script>");
+
+				htmlFile.WriteLine(
+					"var originalStyleElement = null; \n" +
+					"document.addEventListener('DOMContentLoaded', function(event) { regenerateStickyColumnCss(); }) \n" +
+					"window.addEventListener('resize', function(event) { regenerateStickyColumnCss(); }) \n"+
+					"\n"+
+					"function regenerateStickyColumnCss() { \n" +
+					"  var styleElement=document.getElementById('pageStyle'); \n" +
+					"  var table=document.getElementById('mainTable'); \n" +
+					"  if ( table.rows.length < 2 ) \n" +
+					"	return; \n" +
+					"  if (originalStyleElement == null) \n" +
+					"    originalStyleElement = styleElement.textContent; \n" +
+					"  else \n" +
+					"    styleElement.textContent = originalStyleElement  \n"
+				);
+
+				// Make the columns Sticky and compute their X offsets
+				htmlFile.WriteLine(
+					"  var numStickyCols="+ stickyColumnCount + "; \n" +
+					"  var xOffset=0; \n" +
+					"  for (var i=0;i<numStickyCols;i++) \n" +
+					"  { \n"+
+					"	 var rBorderParam=(i==numStickyCols-1) ? 'border-right: 2px solid black;':''; \n" +
+					"	 styleElement.textContent+='tr.lastHeaderRow th:nth-child('+(i+1)+') {  z-index: 8;  border-top: 2px solid black;  font-size: 11px;  left: '+xOffset+'px;'+rBorderParam+'}'; \n" +
+					"	 styleElement.textContent+='td:nth-child('+(i+1)+') {  position: -webkit-sticky;  position: sticky; z-index: 7;  left: '+xOffset+'px; '+rBorderParam+'}'; \n" +
+					"	 xOffset+=table.rows[1].cells[i].offsetWidth; \n" +
+					"  } \n"
+					);
+				// Make all sticky columns except the count column opaque
+				htmlFile.WriteLine(
+					"  for (var i=0;i<numStickyCols-1;i++) \n" +
+					"  { \n" +
+					"    styleElement.textContent+='tr:nth-child(odd) td:nth-child('+(i+1)+') { background-color: #e2e2e2; }'; \n" +
+					"    styleElement.textContent+='tr:nth-child(even) td:nth-child('+(i+1)+') { background-color: #ffffff; }'; \n" +
+					"  } \n" +
+					"} \n"
+				);
+				htmlFile.WriteLine("</script>");
+			}
+
+			htmlFile.WriteLine("<style type='text/css' id='pageStyle'>");
 			htmlFile.WriteLine("p {  font-family: 'Verdana', Times, serif; font-size: 12px }");
 			htmlFile.WriteLine("h3 {  font-family: 'Verdana', Times, serif; font-size: 14px }");
 			htmlFile.WriteLine("h2 {  font-family: 'Verdana', Times, serif; font-size: 16px }");
 			htmlFile.WriteLine("h1 {  font-family: 'Verdana', Times, serif; font-size: 20px }");
 			string tableCss = "";
-			int frozenColumnCount = 0;
 			if (bScrollableTable)
 			{
 				int firstColMaxStringLength = 0;
@@ -3137,21 +3279,15 @@ namespace PerfSummaries
 						firstColMaxStringLength = Math.Max(firstColMaxStringLength,columns[0].GetStringValue(i).Length);
 					}
 				}
-				frozenColumnCount = 1;
-				// Freeze the second column if it's the "count" column
-				if (columns.Count>1 && columns[1].name=="Count" && isCollated)
-				{
-					frozenColumnCount = 2;
-				}
-				int firstColWidth = firstColMaxStringLength * 7;
+				int firstColWidth = (int)(firstColMaxStringLength * 6.5);
 				tableCss =
-					"table {table-layout: fixed;}"+
-					"table, th, td { border: 0px solid black; border-spacing: 0; border-collapse: separate; padding: " + cellPadding + "px; vertical-align: center; font-family: 'Verdana', Times, serif; font-size: 10px;}" +
+					"table {table-layout: fixed;} \n"+
+					"table, th, td { border: 0px solid black; border-spacing: 0; border-collapse: separate; padding: " + cellPadding + "px; vertical-align: center; font-family: 'Verdana', Times, serif; font-size: 10px;} \n" +
 					"td {" +
 					"  border-right: 1px solid black;"+
 					"  max-width: 400;" +
-					"}" +
-					"tr:first-element { border-top: 2px; border-bottom: 2px }"+
+					"} \n" +
+					"tr:first-element { border-top: 2px; border-bottom: 2px } \n" +
 					"th {" +
 					"  width: 75px;" +
 					"  max-width: 400;" +
@@ -3166,90 +3302,66 @@ namespace PerfSummaries
 					"  word-wrap: break-word;" +
 					"  overflow: hidden;" +
 					"  height: 60;" +
-					"}";
-
-				int xPos = 0;
-				for (int i=0;i< frozenColumnCount; i++)
-				{
-					int colIndex = i + 1;
-					tableCss +=
-						"th:nth-child("+colIndex+") {" +
-						"  position: -webkit-sticky;" +
-						"  position: sticky;" +
-						"  z-index: 7;" +
-						"  background-color: #ffffff;" +
-						"  border-right: 2px solid black;" +
-						"  border-top: 2px solid black;" +
-						"  font-size: 11px;" +
-						"  top:0;" +
-						"  left: " + xPos + "px;" +
-						"}";
-
-					tableCss +=
-						"td:nth-child("+colIndex+") {" +
-						"  position: -webkit-sticky;" +
-						"  position: sticky;" +
-						"  border-right: 2px solid black;" +
-						"  z-index: 6;" +
-						"  left: " + xPos + "px;" +
-						" }";
-
-					xPos += firstColWidth + 4 + cellPadding*2;
-				}
-
-				string firstChildCSS =
-				tableCss +=
-					"th:first-child, td:first-child {" +
-					"  border-left: 2px solid black;" +
-					"  width: " + firstColWidth + ";" +
-					"  min-width: " + firstColWidth + ";" +
-					"  max-width: " + firstColWidth + ";" +
-					" }" +
-
-					"tr.sectionStart td {" +
-					"  border-top: 2px solid black;" +
-					"}";
+					"} \n";
+
+				// Top-left cell of the table is always on top, big font, thick border
+				tableCss += "tr:first-child th:first-child { z-index: 100;  border-right: 2px solid black; border-top: 2px solid black; font-size: 11px; top:0; left: 0px; } \n";
+
+				// Fix the first column width
+				tableCss += "th:first-child, td:first-child { border-left: 2px solid black; min-width: " + firstColWidth + ";} \n";
 
 				if (bAddMinMaxColumns && isCollated)
 				{
-					tableCss += "tr.lastHeaderRow th { top:60px; height:20px; }";
-				}
-				if (bAddMinMaxColumns && isCollated)
-				{
-					// The top left cell is merged, so make sure the one to the right is not sticky horizontally
-					tableCss += "tr:first-child th:nth-child(2) { left:auto; z-index:5; } ";
+					tableCss += "tr.lastHeaderRow th { top:60px; height:20px; } \n";
 				}
 
 				if (!isCollated)
 				{
-					tableCss += "td { max-height: 40px; height:40px } ";
+					tableCss += "td { max-height: 40px; height:40px } \n";
 				}
-				tableCss += "tr:last-child td{border-bottom: 2px solid black;} ";
+				tableCss += "tr:last-child td{border-bottom: 2px solid black;} \n";
 
 			}
 			else 
 			{
 				tableCss =
-					"table, th, td { border: 2px solid black; border-collapse: collapse; padding: "+ cellPadding + "px; vertical-align: center; font-family: 'Verdana', Times, serif; font-size: 11px;}";
+					"table, th, td { border: 2px solid black; border-collapse: collapse; padding: "+ cellPadding + "px; vertical-align: center; font-family: 'Verdana', Times, serif; font-size: 11px;} \n";
 			}
 
 
 			bool bOddRowsGray = !(!bAddMinMaxColumns || !isCollated);
-			tableCss += "tr:nth-child("+ (bOddRowsGray ? "odd" : "even") + ") {background-color: #e2e2e2;} ";
-			tableCss += "tr:nth-child("+ (bOddRowsGray ? "even" : "odd") + ") {background-color: #ffffff;} ";
-			tableCss += "tr:first-child {background-color: #ffffff;} ";
-			tableCss += "tr.lastHeaderRow th { border-bottom: 2px solid black; }";
+			tableCss += "tr:nth-child("+ (bOddRowsGray ? "odd" : "even") + ") {background-color: #e2e2e2;} \n";
+			tableCss += "tr:nth-child("+ (bOddRowsGray ? "even" : "odd") + ") {background-color: #ffffff;} \n";
+			tableCss += "tr:first-child {background-color: #ffffff;} \n";
+			tableCss += "tr.lastHeaderRow th { border-bottom: 2px solid black; } \n";
+
+			// Section start row styles
+			tableCss += "tr.sectionStartMajor td { border-top: 2px solid black; } \n";
+			tableCss += "tr.sectionStartMinor td { border-top: 1px dashed black; } \n";
 
 			htmlFile.WriteLine(tableCss);
 
 			htmlFile.WriteLine("</style>");
 			htmlFile.WriteLine("</head><body>");
-			htmlFile.WriteLine("<table>");
+			htmlFile.WriteLine("<table id='mainTable'>");
+
+			// Get format info for the columns
+			Dictionary<SummaryTableColumn, SummaryTableColumnFormatInfo> columnFormatInfoLookup = new Dictionary<SummaryTableColumn, SummaryTableColumnFormatInfo>();
+			foreach (SummaryTableColumn column in columns)
+			{
+				columnFormatInfoLookup[column] = ( columnFormatInfoList != null ) ? columnFormatInfoList.GetFormatInfo(column.name) : SummaryTableColumnFormatInfoCollection.DefaultColumnInfo;
+			}
 
 			// Automatically colourize the table
 			if (bScrollableTable)
 			{
-				AutoColorizeColumns(summaryTableLowIsBadStatList);
+				foreach (SummaryTableColumn column in columns)
+				{
+					if (column.isNumeric)
+					{
+						column.ComputeAutomaticColourThresholds(columnFormatInfoLookup[column].autoColorizeMode);
+					}
+				}
 			}
 
 			string HeaderRow = "";
@@ -3308,44 +3420,92 @@ namespace PerfSummaries
 			}
 			string[] stripeColors = { "'#e2e2e2'", "'#ffffff'" };
 
-			string prevSectionName = "";
-			for ( int i=0; i<rowCount; i++)
+			// Work out which rows are major/minor section boundaries
+			Dictionary<int, int> rowSectionBoundaryType = new Dictionary<int, int>();
+			if (sectionBoundaries != null)
 			{
-				bool sectionStart = false;
-				if (sectionBoundaryInfo != null)
+				foreach (SummarySectionBoundaryInfo sectionBoundaryInfo in sectionBoundaries)
 				{
-					// Work out the section name if we have section boundary info. When it changes, apply the sectionStart CSS class
-					string sectionName = "";
-					if (sectionBoundaryInfo != null && columnLookup.ContainsKey(sectionBoundaryInfo.statName))
+					// Skip this section boundary info if it's not in this table type
+					if (isCollated && !sectionBoundaryInfo.inCollatedTable)
+					{
+						continue;
+					}
+					if (!isCollated && !sectionBoundaryInfo.inFullTable)
+					{
+						continue;
+					}
+					string prevSectionName = "";
+					for (int i = 0; i < rowCount; i++)
 					{
-						SummaryTableColumn col = columnLookup[sectionBoundaryInfo.statName];
-						string currentValue = col.GetStringValue(i);
-						// Only use the start and end tokens if the table is collated
-						if (isCollated)
+						int boundaryType = 0;
+						if (sectionBoundaryInfo != null)
 						{
-							int startTokenIndex = currentValue.IndexOf(sectionBoundaryInfo.startToken);
-							int endTokenIndex = currentValue.IndexOf(sectionBoundaryInfo.endToken);
-							if (startTokenIndex >= 0 && endTokenIndex > startTokenIndex)
+							// Work out the section name if we have section boundary info. When it changes, apply the sectionStart CSS class
+							string sectionName = "";
+							if (sectionBoundaryInfo != null && columnLookup.ContainsKey(sectionBoundaryInfo.statName))
+							{
+								// Get the section name
+								if (!columnLookup.ContainsKey(sectionBoundaryInfo.statName))
+								{
+									continue;
+								}
+								SummaryTableColumn col = columnLookup[sectionBoundaryInfo.statName];
+								sectionName = col.GetStringValue(i);
+
+								// if we have a start token then strip before it
+								if (sectionBoundaryInfo.startToken != null)
+								{
+									int startTokenIndex = sectionName.IndexOf(sectionBoundaryInfo.startToken);
+									if (startTokenIndex != -1)
+									{
+										sectionName = sectionName.Substring(startTokenIndex + sectionBoundaryInfo.startToken.Length);
+									}
+								}
+
+								// if we have an end token then strip after it
+								if (sectionBoundaryInfo.endToken != null)
+								{
+									int endTokenIndex = sectionName.IndexOf(sectionBoundaryInfo.endToken);
+									if (endTokenIndex != -1)
+									{
+										sectionName = sectionName.Substring(0, endTokenIndex);
+									}
+								}
+							}
+							if (sectionName != prevSectionName && i > 0)
 							{
-								sectionName = currentValue.Substring(startTokenIndex + sectionBoundaryInfo.startToken.Length, endTokenIndex - sectionBoundaryInfo.startToken.Length);
+								// Update the row's boundary type info
+								boundaryType = sectionBoundaryInfo.isMinor ? 1 : 2;
+								if (rowSectionBoundaryType.ContainsKey(i))
+								{
+									boundaryType = Math.Max(rowSectionBoundaryType[i], boundaryType);
+								}
+								rowSectionBoundaryType[i] = boundaryType;
 							}
+							prevSectionName = sectionName;
 						}
-						else
-						{
-							sectionName = currentValue;
-						}
-					}
-					if (sectionName != prevSectionName && i>0)
-					{
-						sectionStart = true;
 					}
-					prevSectionName = sectionName;
 				}
 
+			}
+
+			// Add the rows to the table
+			for ( int i=0; i<rowCount; i++)
+			{
 				string rowClassStr = "";
-				if (sectionStart)
+
+				// Is this a major/minor section boundary
+				if (rowSectionBoundaryType.ContainsKey(i))
 				{
-					rowClassStr = " class='sectionStart'";
+					if (rowSectionBoundaryType[i] == 2)
+					{
+						rowClassStr = " class='sectionStartMajor'";
+					}
+					else if (rowSectionBoundaryType[i] == 1)
+					{
+						rowClassStr = " class='sectionStartMinor'";
+					}
 				}
 
 				htmlFile.Write("<tr"+rowClassStr+">");
@@ -3361,18 +3521,19 @@ namespace PerfSummaries
 						{
 							toolTip = column.GetDisplayName();
 						}
-						toolTipString = "title = '" + toolTip + "'";
+						toolTipString = " title='" + toolTip + "'";
 					}
 					string colour = column.GetColour(i);
 
 					// Alternating row colours are normally handled by CSS, but we need to handle it explicitly if we have frozen first columns
-					if (columnIndex < frozenColumnCount && colour == null)
+					if (columnIndex < stickyColumnCount && colour == null)
 					{
 						colour = stripeColors[i % 2];
 					}
-					string bgColorString = (colour==null ? "" : " bgcolor = " + colour);
+					string bgColorString = (colour==null ? "" : " bgcolor=" + colour);
 					bool bold = false;
-					string stringValue = column.GetStringValue(i, true);
+					string numericFormat = columnFormatInfoLookup[column].numericFormat;
+					string stringValue = column.GetStringValue(i, true, numericFormat);
 					if (maxColumnStringLength > 0 && stringValue.Length > maxColumnStringLength)
 					{
 						stringValue = TableUtil.SafeTruncateHtmlTableValue(stringValue, maxColumnStringLength);
@@ -3381,7 +3542,7 @@ namespace PerfSummaries
 					{
 						stringValue = "'" + stringValue;
 					}
-					string columnString = "<td "+ toolTipString + bgColorString+"> " + (bold ? "<b>" : "") + stringValue + (bold ? "</b>" : "") + "</td>";
+					string columnString = "<td"+ toolTipString + bgColorString+"> " + (bold ? "<b>" : "") + stringValue + (bold ? "</b>" : "") + "</td>";
 					htmlFile.Write(columnString);
 					columnIndex++;
 				}
