PerfReportFort 4.52
- Add support for 3 levels of section barrier. Replace the bool "minor" attribute with int "level" attribute which can be 0,1 or 2 (0 is thickest)
- Fix bug where filter order is ignored for columns in the rowSort list in collated summary tables. Columns we're collating by will still appear first, but their relative order is preserved
- Adjust automated test summary table ordering to sort first by release version

diff --git a/Engine/Source/Programs/CSVTools/PerfReportTool/PerfReportTool.cs b/Engine/Source/Programs/CSVTools/PerfReportTool/PerfReportTool.cs
--- a/Engine/Source/Programs/CSVTools/PerfReportTool/PerfReportTool.cs
+++ b/Engine/Source/Programs/CSVTools/PerfReportTool/PerfReportTool.cs
@@ -22,7 +22,7 @@ namespace PerfReportTool
 {
     class Version
     {
-        private static string VersionString = "4.51";
+        private static string VersionString = "4.52";
 
         public static string Get() { return VersionString; }
     };
diff --git a/Engine/Source/Programs/CSVTools/PerfReportTool/perfsummaries.cs b/Engine/Source/Programs/CSVTools/PerfReportTool/perfsummaries.cs
--- a/Engine/Source/Programs/CSVTools/PerfReportTool/perfsummaries.cs
+++ b/Engine/Source/Programs/CSVTools/PerfReportTool/perfsummaries.cs
@@ -2438,19 +2438,19 @@ namespace PerfSummaries
 
 	class SummarySectionBoundaryInfo
 	{
-		public SummarySectionBoundaryInfo( string inStatName, string inStartToken, string inEndToken, bool inIsMinor, bool inInCollatedTable, bool inInFullTable)
+		public SummarySectionBoundaryInfo( string inStatName, string inStartToken, string inEndToken, int inLevel, bool inInCollatedTable, bool inInFullTable)
 		{
 			statName = inStatName;
 			startToken = inStartToken;
 			endToken = inEndToken;
-			isMinor = inIsMinor;
+			level = inLevel;
 			inCollatedTable = inInCollatedTable;
 			inFullTable = inInFullTable;
 		}
 		public string statName;
 		public string startToken;
 		public string endToken;
-		public bool isMinor;
+		public int level;
 		public bool inCollatedTable;
 		public bool inFullTable;
 	};
@@ -2484,7 +2484,7 @@ namespace PerfSummaries
 						sectionBoundaryEl.GetSafeAttibute<string>("statName"),
 						sectionBoundaryEl.GetSafeAttibute<string>("startToken"),
 						sectionBoundaryEl.GetSafeAttibute<string>("endToken"),
-						sectionBoundaryEl.GetSafeAttibute<bool>("minor", false),
+						sectionBoundaryEl.GetSafeAttibute<int>("level", 0),
 						sectionBoundaryEl.GetSafeAttibute<bool>("inCollatedTable", true),
 						sectionBoundaryEl.GetSafeAttibute<bool>("inFullTable", true)
 						);
@@ -2864,23 +2864,38 @@ namespace PerfSummaries
 		{
 			int numSubColumns=addMinMaxColumns ? 3 : 1;
 
-			List<SummaryTableColumn> newColumns = new List<SummaryTableColumn>();
-			List<string> finalSortByList = new List<string>();
+			// Find all the columns in collateByList
+			HashSet<SummaryTableColumn> collateByColumns = new HashSet<SummaryTableColumn>();
 			foreach (string collateBy in collateByList)
 			{
 				string key = collateBy.ToLower();
 				if (columnLookup.ContainsKey(key))
 				{
-					newColumns.Add(new SummaryTableColumn(columnLookup[key].name, false, columnLookup[key].displayName));
-					finalSortByList.Add(key);
+					collateByColumns.Add(columnLookup[key]);					
 				}
 			}
-
-			if ( finalSortByList.Count == 0 )
+			if (collateByColumns.Count == 0)
 			{
 				throw new Exception("None of the metadata strings were found:" + collateByList.ToString());
 			}
 
+			// Add the new collateBy columns in the order they appear in the original column list
+			List<SummaryTableColumn> newColumns = new List<SummaryTableColumn>();
+			List<string> finalSortByList = new List<string>();
+			foreach (SummaryTableColumn srcColumn in columns)
+			{
+				if ( collateByColumns.Contains(srcColumn) )
+				{
+					newColumns.Add(new SummaryTableColumn(srcColumn.name, false, srcColumn.displayName));
+					finalSortByList.Add(srcColumn.name.ToLower());
+					// Early out if we've found all the columns
+					if (finalSortByList.Count == collateByColumns.Count)
+					{
+						break;
+					}
+				}
+			}
+
             newColumns.Add(new SummaryTableColumn("Count", true));
             int countColumnIndex = newColumns.Count-1;
 
@@ -2889,7 +2904,7 @@ namespace PerfSummaries
 			foreach ( SummaryTableColumn column in columns )
 			{
                 // Add avg/min/max columns for this column if it's numeric and we didn't already add it above 
-				if ( column.isNumeric && !finalSortByList.Contains(column.name.ToLower()))
+				if ( column.isNumeric && !collateByColumns.Contains(column))
 				{
 					srcToDestBaseColumnIndex.Add( newColumns.Count );
 					newColumns.Add(new SummaryTableColumn("Avg " + column.name, true));
@@ -3336,8 +3351,9 @@ namespace PerfSummaries
 			tableCss += "tr.lastHeaderRow th { border-bottom: 2px solid black; } \n";
 
 			// Section start row styles
-			tableCss += "tr.sectionStartMajor td { border-top: 2px solid black; } \n";
-			tableCss += "tr.sectionStartMinor td { border-top: 1px dashed black; } \n";
+			tableCss += "tr.sectionStartLevel0 td { border-top: 2px solid black; } \n";
+			tableCss += "tr.sectionStartLevel1 td { border-top: 1px solid black; } \n";
+			tableCss += "tr.sectionStartLevel2 td { border-top: 1px dashed black; } \n";
 
 			htmlFile.WriteLine(tableCss);
 
@@ -3421,7 +3437,7 @@ namespace PerfSummaries
 			string[] stripeColors = { "'#e2e2e2'", "'#ffffff'" };
 
 			// Work out which rows are major/minor section boundaries
-			Dictionary<int, int> rowSectionBoundaryType = new Dictionary<int, int>();
+			Dictionary<int, int> rowSectionBoundaryLevel = new Dictionary<int, int>();
 			if (sectionBoundaries != null)
 			{
 				foreach (SummarySectionBoundaryInfo sectionBoundaryInfo in sectionBoundaries)
@@ -3438,7 +3454,7 @@ namespace PerfSummaries
 					string prevSectionName = "";
 					for (int i = 0; i < rowCount; i++)
 					{
-						int boundaryType = 0;
+						int boundaryLevel = 0;
 						if (sectionBoundaryInfo != null)
 						{
 							// Work out the section name if we have section boundary info. When it changes, apply the sectionStart CSS class
@@ -3476,12 +3492,13 @@ namespace PerfSummaries
 							if (sectionName != prevSectionName && i > 0)
 							{
 								// Update the row's boundary type info
-								boundaryType = sectionBoundaryInfo.isMinor ? 1 : 2;
-								if (rowSectionBoundaryType.ContainsKey(i))
+								boundaryLevel = sectionBoundaryInfo.level;
+								if (rowSectionBoundaryLevel.ContainsKey(i))
 								{
-									boundaryType = Math.Max(rowSectionBoundaryType[i], boundaryType);
+									// Lower level values override higher ones
+									boundaryLevel = Math.Min(rowSectionBoundaryLevel[i], boundaryLevel);
 								}
-								rowSectionBoundaryType[i] = boundaryType;
+								rowSectionBoundaryLevel[i] = boundaryLevel;
 							}
 							prevSectionName = sectionName;
 						}
@@ -3496,15 +3513,12 @@ namespace PerfSummaries
 				string rowClassStr = "";
 
 				// Is this a major/minor section boundary
-				if (rowSectionBoundaryType.ContainsKey(i))
+				if (rowSectionBoundaryLevel.ContainsKey(i))
 				{
-					if (rowSectionBoundaryType[i] == 2)
-					{
-						rowClassStr = " class='sectionStartMajor'";
-					}
-					else if (rowSectionBoundaryType[i] == 1)
+					int sectionLevel = rowSectionBoundaryLevel[i];
+					if (sectionLevel < 3)
 					{
-						rowClassStr = " class='sectionStartMinor'";
+						rowClassStr = " class='sectionStartLevel"+ sectionLevel + "'";
 					}
 				}
 
