PerfReportTool 3.60. Refactoring:
- Split classes into separate files
- Make Summary types self-register so they're self-contained and more easily extendable
- Remove redundant XmlHelper class, convert existing usage of this to use the GetSafeAttribute method
- Output is unchanged (verified with unit tests)

diff --git a/Engine/Source/Programs/CSVTools/PerfReportTool/CsvFileCache..cs b/Engine/Source/Programs/CSVTools/PerfReportTool/CsvFileCache..cs
new file mode 100644
--- /dev/null
+++ b/Engine/Source/Programs/CSVTools/PerfReportTool/CsvFileCache..cs
@@ -0,0 +1,467 @@
+// Copyright Epic Games, Inc. All Rights Reserved.
+
+using System;
+using System.Collections.Generic;
+using System.Linq;
+using System.Text;
+using System.Threading.Tasks;
+using System.Xml;
+using System.Xml.Linq;
+using System.IO;
+using System.Diagnostics;
+using CSVStats;
+using System.Collections;
+using System.Collections.Concurrent;
+using System.Threading;
+using System.Security.Cryptography;
+
+using PerfSummaries;
+
+namespace PerfReportTool
+{
+	class CachedCsvFile
+	{
+		enum FileType
+		{
+			TextCsv,
+			BinaryCsv,
+			SummaryTableCachePrc
+		};
+		public CachedCsvFile(string inFilename, bool useCacheFiles, DerivedMetadataMappings derivedMetadataMappings)
+		{
+			if (inFilename.ToLower().EndsWith(".csv"))
+			{
+				fileType = FileType.TextCsv;
+			}
+			else if (inFilename.ToLower().EndsWith(".csv.bin"))
+			{
+				fileType = FileType.BinaryCsv;
+			}
+			else if (inFilename.ToLower().EndsWith(".prc"))
+			{
+				fileType = FileType.SummaryTableCachePrc;
+			}
+			else
+			{
+				throw new Exception("File extension not supported for file " + inFilename);
+			}
+			string cacheFilename = inFilename + ".cache";
+			if (useCacheFiles && File.Exists(cacheFilename))
+			{
+				string[] fileLines = File.ReadAllLines(cacheFilename);
+
+				// Put the stats and metadata lines in the standard order
+				if (fileLines.Length >= 3)
+				{
+					string metadataLine = fileLines[1];
+					string statsLine = fileLines[2];
+					fileLines[1] = statsLine;
+					fileLines[2] = metadataLine;
+				}
+				dummyCsvStats = CsvStats.ReadCSVFromLines(fileLines, null, 0, true);
+			}
+			else
+			{
+				if (fileType == FileType.BinaryCsv)
+				{
+					dummyCsvStats = CsvStats.ReadBinFile(inFilename, null, 0, true);
+				}
+				else if (fileType == FileType.TextCsv)
+				{
+					textCsvLines = File.ReadAllLines(inFilename);
+					if (textCsvLines.Length > 0)
+					{
+						dummyCsvStats = CsvStats.ReadCSVFromLines(textCsvLines, null, 0, true);
+					}
+					else
+					{
+						Console.WriteLine("CSV file " + inFilename + " is contains no lines!");
+						dummyCsvStats = new CsvStats();
+					}
+				}
+				else if (fileType == FileType.SummaryTableCachePrc)
+				{
+					// Read just the Csv metadata so we can filter
+					metadata = null;
+					cachedSummaryTableRowData = SummaryTableRowData.TryReadFromCacheFile(inFilename, true);
+					if (cachedSummaryTableRowData != null)
+					{
+						foreach (KeyValuePair<string, SummaryTableElement> entry in cachedSummaryTableRowData.dict)
+						{
+							if (entry.Value.type == SummaryTableElement.Type.CsvMetadata)
+							{
+								if (metadata == null)
+								{
+									metadata = new CsvMetadata();
+								}
+								metadata.Values.Add(entry.Key, entry.Value.value);
+							}
+							if (entry.Value.type == SummaryTableElement.Type.CsvStatAverage)
+							{
+								break;
+							}
+						}
+					}
+					else
+					{
+						Console.WriteLine("Invalid PRC file detected: " + inFilename);
+					}
+				}
+			}
+			filename = inFilename;
+
+			// Apply metadata mappings if we have metadata. Note: we skip for PRCs, since these already have mappings cached
+			if (dummyCsvStats != null && dummyCsvStats.metaData != null)
+			{
+				metadata = dummyCsvStats.metaData;
+				derivedMetadataMappings.ApplyMapping(metadata);
+			}
+		}
+
+		public void PrepareCsvData()
+		{
+			if (fileType == FileType.BinaryCsv)
+			{
+				finalCsv = CsvStats.ReadBinFile(filename);
+			}
+			else
+			{
+				if (textCsvLines == null)
+				{
+					textCsvLines = File.ReadAllLines(filename);
+				}
+			}
+		}
+
+		public CsvStats GetFinalCsv()
+		{
+			if (finalCsv != null)
+			{
+				return finalCsv;
+			}
+			if (fileType == FileType.TextCsv)
+			{
+				finalCsv = CsvStats.ReadCSVFromLines(textCsvLines, null);
+				return finalCsv;
+			}
+			return null;
+		}
+
+
+		public bool DoesMetadataMatchQuery(QueryExpression metadataQuery)
+		{
+			if (metadataQuery == null)
+			{
+				return true;
+			}
+			if (metadata == null)
+			{
+				Console.WriteLine("CSV " + filename + " has no metadata");
+				return false;
+			}
+			return metadataQuery.Evaluate(metadata);
+		}
+
+		public void ComputeSummaryTableCacheId(string reportTypeId)
+		{
+			// If the CSV has an embedded ID, use that. Otherwise generate one. 
+			string csvId;
+			if (metadata != null && metadata.Values.ContainsKey("csvid"))
+			{
+				csvId = metadata.Values["csvid"];
+			}
+			else
+			{
+				// Fall back to absolute path if the CSVID metadata doesn't exist
+				// Note that using the filename like this means moving the file will result in a new entry for each location
+				StringBuilder sb = new StringBuilder();
+				sb.Append("CSVFILENAME={" + Path.GetFullPath(filename).ToLower() + "}\n");
+				if (metadata != null)
+				{
+					foreach (string key in metadata.Values.Keys)
+					{
+						sb.Append("{" + key + "}={" + metadata.Values[key] + "}\n");
+					}
+				}
+				csvId = HashHelper.StringToHashStr(sb.ToString());
+			}
+			if (reportTypeId == null)
+			{
+				summaryTableCacheId = csvId;
+			}
+			else
+			{
+				summaryTableCacheId = csvId + "_" + reportTypeId;
+			}
+		}
+
+
+		public string filename;
+		public string summaryTableCacheId;
+		public string[] textCsvLines;
+		public CsvStats dummyCsvStats;
+		public CsvStats finalCsv;
+		public CsvMetadata metadata;
+		public SummaryTableRowData cachedSummaryTableRowData;
+		public ReportTypeInfo reportTypeInfo;
+		FileType fileType;
+	};
+
+	class CsvFileCache
+	{
+		public CsvFileCache(
+			string[] inCsvFilenames,
+			int inLookaheadCount,
+			int inThreadCount,
+			bool inUseCacheFiles,
+			QueryExpression inMetadataQuery,
+			ReportXML inReportXml,
+			ReportTypeParams inReportTypeParams,
+			bool inBulkMode,
+			bool inSummaryTableCacheOnlyMode,
+			bool inSummaryTableCacheUseOnlyCsvID,
+			bool inRemoveDuplicates,
+			bool inRequireMetadata,
+			string inSummaryTableCacheDir,
+			bool inListFilesMode)
+		{
+			csvFileInfos = new CsvFileInfo[inCsvFilenames.Length];
+			for (int i = 0; i < inCsvFilenames.Length; i++)
+			{
+				csvFileInfos[i] = new CsvFileInfo(inCsvFilenames[i]);
+			}
+			fileCache = this;
+			writeIndex = 0;
+			useCacheFiles = inUseCacheFiles;
+			readIndex = 0;
+			lookaheadCount = inLookaheadCount;
+			countFreedSinceLastGC = 0;
+			reportXml = inReportXml;
+			bulkMode = inBulkMode;
+			summaryTableCacheOnlyMode = inSummaryTableCacheOnlyMode;
+			summaryTableCacheUseOnlyCsvID = inSummaryTableCacheUseOnlyCsvID;
+			metadataQuery = inMetadataQuery;
+			derivedMetadataMappings = inReportXml.derivedMetadataMappings;
+			summaryTableCacheDir = inSummaryTableCacheDir;
+			reportTypeParams = inReportTypeParams;
+			bRemoveDuplicates = inRemoveDuplicates;
+			bRequireMetadata = inRequireMetadata;
+			bListFilesMode = inListFilesMode;
+
+			csvIdToFilename = new Dictionary<string, string>();
+			outCsvQueue = new BlockingCollection<CsvFileInfo>(lookaheadCount);
+
+			// Kick off the workers (must be done last)
+			if (inLookaheadCount > 0)
+			{
+				Console.WriteLine("Kicking off " + inThreadCount + " precache threads with lookahead " + lookaheadCount);
+				precacheThreads = new Thread[inThreadCount];
+				precacheJobs = new ThreadStart[inThreadCount];
+				for (int i = 0; i < precacheThreads.Length; i++)
+				{
+					precacheJobs[i] = new ThreadStart(PrecacheThreadRun);
+					precacheThreads[i] = new Thread(precacheJobs[i]);
+					precacheThreads[i].Start();
+				}
+			}
+		}
+
+		public CachedCsvFile GetNextCachedCsvFile()
+		{
+			CachedCsvFile file = null;
+			if (readIndex >= csvFileInfos.Length)
+			{
+				// We're done
+				return null;
+			}
+
+			// Find the next valid fileinfo
+			if (precacheThreads == null)
+			{
+				CsvFileInfo fileInfo = csvFileInfos[readIndex];
+				file = new CachedCsvFile(fileInfo.filename, useCacheFiles, derivedMetadataMappings);
+				if (!file.DoesMetadataMatchQuery(metadataQuery)) // TODO do we need to check this here and in the thread?
+				{
+					return null;
+				}
+			}
+			else
+			{
+				CsvFileInfo fileInfo = outCsvQueue.Take();
+				file = fileInfo.cachedFile;
+				fileInfo.cachedFile = null;
+				countFreedSinceLastGC++;
+				// Periodically GC
+				if (countFreedSinceLastGC > 16 && summaryTableCacheOnlyMode == false)
+				{
+					GC.Collect();
+					GC.WaitForPendingFinalizers();
+					countFreedSinceLastGC = 0;
+				}
+				if (!fileInfo.isValid)
+				{
+					file = null;
+				}
+			}
+			readIndex++;
+			return file;
+		}
+
+		static void PrecacheThreadRun()
+		{
+			fileCache.ThreadRun();
+		}
+
+		void ThreadRun()
+		{
+			int threadWriteIndex = 0;
+			while (true)
+			{
+				threadWriteIndex = Interlocked.Increment(ref writeIndex) - 1;
+				if (threadWriteIndex >= csvFileInfos.Length)
+				{
+					// We're done
+					break;
+				}
+				CsvFileInfo fileInfo = csvFileInfos[threadWriteIndex];
+
+				// Process the file
+				CachedCsvFile file = new CachedCsvFile(fileInfo.filename, useCacheFiles, derivedMetadataMappings);
+
+				bool bProcessThisFile = true;
+				if (bRequireMetadata && file.metadata == null)
+				{
+					Console.WriteLine("CSV has no metadata. Skipping: " + fileInfo.filename);
+					bProcessThisFile = false;
+				}
+				else if (bRemoveDuplicates && file.metadata != null && file.metadata.Values.ContainsKey("csvid"))
+				{
+					string csvId = file.metadata.Values["csvid"];
+					lock (csvIdToFilename)
+					{
+						if (csvIdToFilename.ContainsKey(csvId))
+						{
+							Console.WriteLine("Duplicate CSV found: " + fileInfo.filename);
+							Console.WriteLine("   First version   : " + csvIdToFilename[csvId]);
+							bProcessThisFile = false;
+							duplicateCount++;
+						}
+						else
+						{
+							csvIdToFilename[csvId] = fileInfo.filename;
+						}
+					}
+				}
+
+				if (bProcessThisFile && file.DoesMetadataMatchQuery(metadataQuery))
+				{
+					if (bListFilesMode)
+					{
+						Console.WriteLine("File: " + fileInfo.filename);
+					}
+					else if (summaryTableCacheOnlyMode)
+					{
+						// Assume the report type is valid if we have cached metadata, since we don't actually need the reporttype in metadata only mode
+						if (file.cachedSummaryTableRowData != null)
+						{
+							// Just read the full cached metadata
+							file.cachedSummaryTableRowData = SummaryTableRowData.TryReadFromCacheFile(fileInfo.filename);
+							fileInfo.isValid = true;
+						}
+					}
+					else
+					{
+						file.reportTypeInfo = GetCsvReportTypeInfo(file, bulkMode);
+						if (file.reportTypeInfo != null)
+						{
+							string reportTypeHash = summaryTableCacheUseOnlyCsvID ? null : file.reportTypeInfo.GetSummaryTableCacheID();
+							file.ComputeSummaryTableCacheId(reportTypeHash);
+							if (summaryTableCacheDir != null)
+							{
+								// If a summary metadata cache is specified, try reading from it instead of reading the whole CSV
+								// Note that this will be disabled if we're not in bulk mode
+								file.cachedSummaryTableRowData = SummaryTableRowData.TryReadFromCache(summaryTableCacheDir, file.summaryTableCacheId);
+								if (file.cachedSummaryTableRowData == null)
+								{
+									Console.WriteLine("Failed to read summary metadata from cache for CSV: " + fileInfo.filename);
+								}
+							}
+							if (file.cachedSummaryTableRowData == null)
+							{
+								file.PrepareCsvData();
+							}
+
+							// Only read the full file data if the metadata matches
+							fileInfo.isValid = true;
+						}
+					}
+				}
+				fileInfo.cachedFile = file;
+				// TODO: only add valid files to the queue? Need alternative stopping condition
+				outCsvQueue.Add(fileInfo);
+			}
+		}
+
+		ReportTypeInfo GetCsvReportTypeInfo(CachedCsvFile csvFile, bool bBulkMode)
+		{
+			try
+			{
+				return reportXml.GetReportTypeInfo(reportTypeParams.reportTypeOverride, csvFile, bBulkMode, reportTypeParams.forceReportType);
+			}
+			catch (Exception e)
+			{
+				if (bBulkMode)
+				{
+					Console.Out.WriteLine("[ERROR] : " + e.Message);
+					return null;
+				}
+				else
+				{
+					// If we're not in bulk mode, exceptions are fatal
+					throw e;
+				}
+			}
+		}
+
+
+		ThreadStart[] precacheJobs;
+		Thread[] precacheThreads;
+		int writeIndex;
+		int readIndex;
+		int lookaheadCount;
+		int countFreedSinceLastGC;
+		bool useCacheFiles;
+		bool bulkMode;
+		bool bRequireMetadata;
+		bool summaryTableCacheOnlyMode;
+		bool summaryTableCacheUseOnlyCsvID;
+		QueryExpression metadataQuery;
+		string summaryTableCacheDir;
+		ReportXML reportXml;
+		ReportTypeParams reportTypeParams;
+		DerivedMetadataMappings derivedMetadataMappings;
+
+		Dictionary<string, string> csvIdToFilename;
+		public int duplicateCount = 0;
+		bool bRemoveDuplicates;
+		bool bListFilesMode;
+
+		static CsvFileCache fileCache;
+		BlockingCollection<CsvFileInfo> outCsvQueue;
+
+		class CsvFileInfo
+		{
+			public CsvFileInfo(string inFilename)
+			{
+				filename = inFilename;
+				cachedFile = null;
+				isValid = false;
+			}
+			public CachedCsvFile cachedFile;
+			public string filename;
+			public bool isValid;
+		}
+		CsvFileInfo[] csvFileInfos;
+
+	}
+}
\ No newline at end of file
diff --git a/Engine/Source/Programs/CSVTools/PerfReportTool/PerfReportTool.cs b/Engine/Source/Programs/CSVTools/PerfReportTool/PerfReportTool.cs
--- a/Engine/Source/Programs/CSVTools/PerfReportTool/PerfReportTool.cs
+++ b/Engine/Source/Programs/CSVTools/PerfReportTool/PerfReportTool.cs
@@ -5,15 +5,11 @@ using System;
 using System.Collections.Generic;
 using System.Linq;
 using System.Text;
-using System.Threading.Tasks;
-using System.Xml;
 using System.Xml.Linq;
 using System.IO;
 using System.Diagnostics;
 using CSVStats;
 using System.Collections;
-using System.Collections.Concurrent;
-using System.Threading;
 using System.Security.Cryptography;
 
 using PerfSummaries;
@@ -22,7 +18,7 @@ namespace PerfReportTool
 {
     class Version
     {
-        private static string VersionString = "4.52";
+        private static string VersionString = "4.60";
 
         public static string Get() { return VersionString; }
     };
@@ -49,387 +45,6 @@ namespace PerfReportTool
 		}
 	}
 
-	class XmlHelper
-    { 
-        public static int ReadAttributeInt(XElement element, string AttributeName, int DefaultValue)
-        {
-            try
-            {
-                return Convert.ToInt32(element.Attribute(AttributeName).Value);
-            }
-            catch
-            {
-            }
-            return DefaultValue;
-        }
-
-        public static double ReadAttributeDouble(XElement element, string AttributeName, double DefaultValue)
-        {
-            try
-            {
-                if (element.Attribute(AttributeName) != null)
-                {
-                    return Convert.ToDouble(element.Attribute(AttributeName).Value, System.Globalization.CultureInfo.InvariantCulture);
-                }
-            }
-            catch { }
-            return DefaultValue;
-        }
-
-        public static bool ReadAttributeBool(XElement element, string AttributeName, bool DefaultValue )
-        {
-            try
-            {
-                if (element.Attribute(AttributeName) != null)
-                {
-                    return Convert.ToInt32(element.Attribute(AttributeName).Value) == 1;
-                }
-            }
-            catch { }
-            return DefaultValue;
-        }
-
-        public static string ReadAttribute(XElement element, string AttributeName, string DefaultValue )
-        {
-            if (element.Attribute(AttributeName) != null)
-            {
-                return element.Attribute(AttributeName).Value;
-            }
-            return DefaultValue;
-        }
-    };
-
-	class OptionalString
-    {
-        public OptionalString(string valueIn)
-        {
-            value = valueIn;
-            isSet = true;
-        }
-        public OptionalString()
-        {
-            isSet = false;
-        }
-        public OptionalString( XElement element, string Name, bool IsElement = false )
-        {
-            isSet = false;
-            if (IsElement )
-            {
-                XElement child = element.Element(Name);
-                if (child != null)
-                {
-                    value = child.Value;
-                    isSet = true;
-                }
-            }
-            else
-            {
-                XAttribute child = element.Attribute(Name);
-                if (child != null)
-                {
-                    value = child.Value;
-                    isSet = true;
-                }
-            }
-        }
-
-        public void InheritFrom(OptionalString baseVersion) { if (!isSet) { isSet = baseVersion.isSet; value = baseVersion.value; } }
-        public bool isSet;
-        public string value;
-    };
-
-    class OptionalBool
-    {
-        public OptionalBool(bool valueIn)
-        {
-            value = valueIn;
-            isSet = true;
-        }
-        public OptionalBool()
-        {
-            isSet = false;
-        }
-        public OptionalBool(XElement element, string AttributeName)
-        {
-            isSet = false;
-            try
-            {
-                if (element.Attribute(AttributeName) != null)
-                {
-                    value = Convert.ToInt32(element.Attribute(AttributeName).Value) == 1;
-                    isSet = true;
-                }
-            }
-            catch {}
-        }
-        public void InheritFrom(OptionalBool baseVersion) { if (!isSet) { isSet = baseVersion.isSet; value = baseVersion.value; } }
-
-        public bool isSet;
-        public bool value;
-    };
-
-    class OptionalInt
-    {
-        public OptionalInt(int valueIn)
-        {
-            value = valueIn;
-            isSet = true;
-        }
-        public OptionalInt()
-        {
-            isSet = false;
-        }
-        public OptionalInt(XElement element, string AttributeName)
-        {
-            isSet = false;
-            try
-            {
-                if (element.Attribute(AttributeName) != null)
-                {
-                    value = Convert.ToInt32(element.Attribute(AttributeName).Value);
-                    isSet = true;
-                }
-            }
-            catch { }
-        }
-        public void InheritFrom(OptionalInt baseVersion) { if (!isSet) { isSet = baseVersion.isSet; value = baseVersion.value; } }
-
-        public bool isSet;
-        public int value;
-    };
-
-    class OptionalDouble
-    {
-        public OptionalDouble(int valueIn)
-        {
-            value = valueIn;
-            isSet = true;
-        }
-        public OptionalDouble()
-        {
-            isSet = false;
-        }
-        public OptionalDouble(XElement element, string AttributeName)
-        {
-            isSet = false;
-            try
-            {
-                if (element.Attribute(AttributeName) != null)
-                {
-                    value = Convert.ToDouble(element.Attribute(AttributeName).Value, System.Globalization.CultureInfo.InvariantCulture);
-                    isSet = true;
-                }
-            }
-            catch {}
-        }
-
-        public void InheritFrom(OptionalDouble baseVersion) { if (!isSet) { isSet = baseVersion.isSet; value = baseVersion.value; } }
-
-        public bool isSet;
-        public double value;
-    };
-
-	static class OptionalHelper
-	{
-		public static string GetDoubleSetting(OptionalDouble setting, string cmdline)
-		{
-			return (setting.isSet ? (cmdline + setting.value.ToString()) : "");
-		}
-
-		public static string GetStringSetting(OptionalString setting, string cmdline)
-		{
-			return (setting.isSet ? (cmdline + setting.value) : "");
-		}
-	};
-
-    class CachedCsvFile
-    {
-		enum FileType
-		{
-			TextCsv,
-			BinaryCsv,
-			SummaryTableCachePrc
-		};
-		public CachedCsvFile(string inFilename, bool useCacheFiles, DerivedMetadataMappings derivedMetadataMappings)
-		{
-			if (inFilename.ToLower().EndsWith(".csv"))
-			{
-				fileType = FileType.TextCsv;
-			}
-			else if (inFilename.ToLower().EndsWith(".csv.bin"))
-			{
-				fileType = FileType.BinaryCsv;
-			}
-			else if (inFilename.ToLower().EndsWith(".prc"))
-			{
-				fileType = FileType.SummaryTableCachePrc;
-			}
-			else
-			{
-				throw new Exception("File extension not supported for file " + inFilename);
-			}
-			string cacheFilename = inFilename + ".cache";
-			if (useCacheFiles && File.Exists(cacheFilename))
-			{
-				string[] fileLines = File.ReadAllLines(cacheFilename);
-
-				// Put the stats and metadata lines in the standard order
-				if (fileLines.Length >= 3)
-				{
-					string metadataLine = fileLines[1];
-					string statsLine = fileLines[2];
-					fileLines[1] = statsLine;
-					fileLines[2] = metadataLine;
-				}
-				dummyCsvStats = CsvStats.ReadCSVFromLines(fileLines, null, 0, true);
-			}
-			else
-			{
-				if (fileType == FileType.BinaryCsv)
-				{
-					dummyCsvStats = CsvStats.ReadBinFile(inFilename, null, 0, true);
-				}
-				else if (fileType == FileType.TextCsv)
-				{
-					textCsvLines = File.ReadAllLines(inFilename);
-					if (textCsvLines.Length > 0)
-					{
-						dummyCsvStats = CsvStats.ReadCSVFromLines(textCsvLines, null, 0, true);
-					}
-					else
-					{
-						Console.WriteLine("CSV file " + inFilename + " is contains no lines!");
-						dummyCsvStats = new CsvStats();
-					}
-				}
-				else if (fileType == FileType.SummaryTableCachePrc)
-				{
-					// Read just the Csv metadata so we can filter
-					metadata = null;
-					cachedSummaryTableRowData = SummaryTableRowData.TryReadFromCacheFile(inFilename, true);
-					if (cachedSummaryTableRowData != null)
-					{
-						foreach (KeyValuePair<string, SummaryTableElement> entry in cachedSummaryTableRowData.dict)
-						{
-							if (entry.Value.type == SummaryTableElement.Type.CsvMetadata)
-							{
-								if (metadata == null)
-								{
-									metadata = new CsvMetadata();
-								}
-								metadata.Values.Add(entry.Key, entry.Value.value);
-							}
-							if (entry.Value.type == SummaryTableElement.Type.CsvStatAverage)
-							{
-								break;
-							}
-						}
-					}
-					else
-					{
-						Console.WriteLine("Invalid PRC file detected: " + inFilename);
-					}
-				}
-			}
-			filename = inFilename;
-
-			// Apply metadata mappings if we have metadata. Note: we skip for PRCs, since these already have mappings cached
-            if (dummyCsvStats != null && dummyCsvStats.metaData != null)
-            {
-				metadata = dummyCsvStats.metaData;
-				derivedMetadataMappings.ApplyMapping(metadata);
-			}
-		}
-
-		public void PrepareCsvData()
-		{
-			if (fileType == FileType.BinaryCsv)
-			{
-				finalCsv = CsvStats.ReadBinFile(filename);
-			}
-			else
-			{
-				if (textCsvLines == null)
-				{
-					textCsvLines = File.ReadAllLines(filename);
-				}
-			}
-		}
-
-		public CsvStats GetFinalCsv()
-		{
-			if (finalCsv != null)
-			{
-				return finalCsv;
-			}
-			if (fileType == FileType.TextCsv)
-			{
-				finalCsv = CsvStats.ReadCSVFromLines(textCsvLines, null);
-				return finalCsv;
-			}
-			return null;
-		}
-
-
-		public bool DoesMetadataMatchQuery(QueryExpression metadataQuery)
-		{
-			if (metadataQuery == null)
-			{
-				return true;
-			}
-			if (metadata == null)
-			{
-				Console.WriteLine("CSV " + filename + " has no metadata");
-				return false;
-			}
-			return metadataQuery.Evaluate(metadata);
-		}
-
-		public void ComputeSummaryTableCacheId(string reportTypeId)
-		{
-			// If the CSV has an embedded ID, use that. Otherwise generate one. 
-			string csvId;
-			if (metadata != null && metadata.Values.ContainsKey("csvid"))
-			{
-				csvId = metadata.Values["csvid"];
-			}
-			else
-			{
-				// Fall back to absolute path if the CSVID metadata doesn't exist
-				// Note that using the filename like this means moving the file will result in a new entry for each location
-				StringBuilder sb = new StringBuilder();
-				sb.Append("CSVFILENAME={" + Path.GetFullPath(filename).ToLower() + "}\n");
-				if (metadata != null)
-				{
-					foreach (string key in metadata.Values.Keys)
-					{
-						sb.Append("{" + key + "}={" + metadata.Values[key] + "}\n");
-					}
-				}
-				csvId = HashHelper.StringToHashStr(sb.ToString());
-			}
-			if ( reportTypeId == null )
-			{
-				summaryTableCacheId = csvId;
-			}
-			else
-			{
-				summaryTableCacheId = csvId + "_" + reportTypeId;
-			}
-		}
-
-
-		public string filename;
-		public string summaryTableCacheId;
-        public string[] textCsvLines;
-        public CsvStats dummyCsvStats;
-		public CsvStats finalCsv;
-		public CsvMetadata metadata;
-		public SummaryTableRowData cachedSummaryTableRowData;
-		public ReportTypeInfo reportTypeInfo;
-		FileType fileType;
-	};
-
 	class SummaryTableCacheStats
 	{
 		public int WriteCount = 0;
@@ -454,7 +69,6 @@ namespace PerfReportTool
 		}
 	};
 
-
 	class Program : CommandLineTool
     {
 		static string formatString =
@@ -583,6 +197,9 @@ namespace PerfReportTool
 
 			ReadCommandLine(args);
             PerfLog perfLog = new PerfLog(GetBoolArg("perfLog"));
+
+			SummaryFactory.Init();
+
 			string csvDir = null;
 
 			bool bBulkMode = false;
@@ -1894,192 +1511,6 @@ namespace PerfReportTool
 
     }
 
-    class ReportTypeInfo
-    {
-        public ReportTypeInfo(XElement element, Dictionary<string,XElement> sharedSummaries, string baseXmlDirectory)
-        {	
-            graphs = new List<ReportGraph>();
-            summaries = new List<Summary>();
-			name = element.Attribute("name").Value;
-			title = element.Attribute("title").Value;
-			bStripEvents = XmlHelper.ReadAttributeBool(element, "stripEvents", true);
-            foreach (XElement child in element.Elements())
-            {
-				if (child.Name == "graph")
-				{
-					ReportGraph graph = new ReportGraph(child);
-					graphs.Add(graph);
-				}
-				else if (child.Name == "summary" || child.Name=="summaryRef")
-				{
-					XElement summaryElement = null;
-					if (child.Name == "summaryRef")
-					{
-						summaryElement = sharedSummaries[child.Attribute("name").Value];
-					}
-					else
-					{
-						summaryElement = child;
-					}
-					string summaryType = summaryElement.Attribute("type").Value;
-					summaries.Add(SummaryFactory.Create(summaryType, summaryElement, baseXmlDirectory));
-				}
-				else if (child.Name == "metadataToShow")
-				{
-					metadataToShowList = child.Value.Split(',');
-				}
-
-            }
-			ComputeSummaryTableCacheID();
-		}
-
-		public string GetSummaryTableCacheID()
-		{
-			return summaryTableCacheID;
-		}
-
-		private void ComputeSummaryTableCacheID()
-		{
-			StringBuilder sb = new StringBuilder();
-			sb.Append("NAME={" + name + "}\n");
-			sb.Append("TITLE={" + title + "}\n");
-			foreach (Summary summary in summaries)
-			{
-				sb.Append("SUMMARY={");
-				sb.Append("TYPE={" + summary.GetType().ToString() + "}");
-				sb.Append("UNSTRIPPED={" + summary.useUnstrippedCsvStats + "}");
-				sb.Append("STATS={" + string.Join(",", summary.stats) + "}");
-				sb.Append("}");
-			}
-			summaryTableCacheID = HashHelper.StringToHashStr(sb.ToString(),16);
-		}
-
-		public List<ReportGraph> graphs;
-        public List<Summary> summaries;
-		public string name;
-		public string title;
-		public string [] metadataToShowList;
-		public string summaryTableCacheID;
-		public bool bStripEvents;
-	};
-
-
-    class ReportGraph
-    {
-        public ReportGraph(XElement element)
-        {
-            title = element.Attribute("title").Value;
-            budget = new OptionalDouble(element, "budget");
-			inSummary = XmlHelper.ReadAttributeBool(element, "inSummary", false);
-			isExternal = XmlHelper.ReadAttributeBool(element, "external", false);
-			minFilterStatValue = new OptionalDouble(element, "minFilterStatValue");
-		}
-        public string title;
-        public OptionalDouble budget;
-        public bool inSummary;
-		public bool isExternal;
-		public OptionalDouble minFilterStatValue;
-        public GraphSettings settings;
-    };
-
-    class GraphSettings
-    {
-        public GraphSettings(XElement element)
-        {
-            smooth = new OptionalBool(element, "smooth");
-            thickness = new OptionalDouble(element, "thickness");
-			miny = new OptionalDouble(element, "miny");
-            maxy = new OptionalDouble(element, "maxy");
-			threshold = new OptionalDouble(element, "threshold");
-			averageThreshold = new OptionalDouble(element, "averageThreshold");
-			minFilterStatValue = new OptionalDouble(element, "minFilterStatValue");
-			minFilterStatName = new OptionalString(element, "minFilterStatName");
-			smoothKernelPercent = new OptionalDouble(element, "smoothKernelPercent");
-            smoothKernelSize = new OptionalDouble(element, "smoothKernelSize");
-            compression = new OptionalDouble(element, "compression");
-            width = new OptionalInt(element, "width");
-            height = new OptionalInt(element, "height");
-            stacked = new OptionalBool(element, "stacked");
-            showAverages = new OptionalBool(element, "showAverages");
-            filterOutZeros = new OptionalBool(element, "filterOutZeros");
-            maxHierarchyDepth = new OptionalInt(element, "maxHierarchyDepth");
-            hideStatPrefix = new OptionalString(element, "hideStatPrefix");
-            mainStat = new OptionalString(element, "mainStat");
-            showEvents = new OptionalString(element, "showEvents");
-            requiresDetailedStats = new OptionalBool(element, "requiresDetailedStats");
-            ignoreStats = new OptionalString(element, "ignoreStats");
-
-            statString = new OptionalString(element, "statString", true);
-            additionalArgs = new OptionalString(element, "additionalArgs", true);
-            statMultiplier = new OptionalDouble(element, "statMultiplier");
-			legendAverageThreshold = new OptionalDouble(element, "legendAverageThreshold");
-			snapToPeaks = new OptionalBool(element, "snapToPeaks");
-			lineDecimalPlaces = new OptionalInt(element, "lineDecimalPlaces");
-		}
-		public void InheritFrom(GraphSettings baseSettings)
-        {
-            smooth.InheritFrom(baseSettings.smooth);
-            statString.InheritFrom(baseSettings.statString);
-            thickness.InheritFrom(baseSettings.thickness);
-			miny.InheritFrom(baseSettings.miny);
-            maxy.InheritFrom(baseSettings.maxy);
-			threshold.InheritFrom(baseSettings.threshold);
-			averageThreshold.InheritFrom(baseSettings.averageThreshold);
-			minFilterStatValue.InheritFrom(baseSettings.minFilterStatValue);
-			minFilterStatName.InheritFrom(baseSettings.minFilterStatName);
-            smoothKernelSize.InheritFrom(baseSettings.smoothKernelSize);
-            smoothKernelPercent.InheritFrom(baseSettings.smoothKernelPercent);
-            compression.InheritFrom(baseSettings.compression);
-            width.InheritFrom(baseSettings.width);
-            height.InheritFrom(baseSettings.height);
-            additionalArgs.InheritFrom(baseSettings.additionalArgs);
-            stacked.InheritFrom(baseSettings.stacked);
-            showAverages.InheritFrom(baseSettings.showAverages);
-            filterOutZeros.InheritFrom(baseSettings.filterOutZeros);
-            maxHierarchyDepth.InheritFrom(baseSettings.maxHierarchyDepth);
-            hideStatPrefix.InheritFrom(baseSettings.hideStatPrefix);
-            mainStat.InheritFrom(baseSettings.mainStat);
-            showEvents.InheritFrom(baseSettings.showEvents);
-            requiresDetailedStats.InheritFrom(baseSettings.requiresDetailedStats);
-            statMultiplier.InheritFrom(baseSettings.statMultiplier);
-            ignoreStats.InheritFrom(baseSettings.ignoreStats);
-			legendAverageThreshold.InheritFrom(baseSettings.legendAverageThreshold);
-			snapToPeaks.InheritFrom(baseSettings.snapToPeaks);
-			lineDecimalPlaces.InheritFrom(baseSettings.lineDecimalPlaces);
-
-		}
-        public OptionalBool smooth;
-        public OptionalString statString;
-        public OptionalDouble thickness;
-		public OptionalDouble miny;
-        public OptionalDouble maxy;
-		public OptionalDouble threshold;
-		public OptionalDouble averageThreshold;
-		public OptionalDouble minFilterStatValue;
-		public OptionalString minFilterStatName;
-		public OptionalDouble smoothKernelSize;
-        public OptionalDouble smoothKernelPercent;
-        public OptionalDouble compression;
-        public OptionalInt width;
-        public OptionalInt height;
-        public OptionalString additionalArgs;
-        public OptionalBool stacked;
-        public OptionalBool showAverages;
-        public OptionalBool filterOutZeros;
-        public OptionalInt maxHierarchyDepth;
-        public OptionalString hideStatPrefix;
-        public OptionalString mainStat;
-        public OptionalString showEvents;
-        public OptionalString ignoreStats;
-        public OptionalDouble statMultiplier;
-		public OptionalDouble legendAverageThreshold;
-
-		public OptionalBool requiresDetailedStats;
-		public OptionalBool snapToPeaks;
-		public OptionalInt lineDecimalPlaces;
-
-	};
-
 	static class Extensions
 	{
 		public static T GetSafeAttibute<T>(this XElement element, string attributeName, T defaultValue = default(T))
@@ -2100,719 +1531,9 @@ namespace PerfReportTool
 		}
 
 
-    };
-
-    class CsvEventStripInfo
-    {
-        public string beginName;
-        public string endName;
     };
 
 
-	class DerivedMetadataEntry
-	{
-		public DerivedMetadataEntry(string inSourceName, string inSourceValue, string inDestName, string inDestValue )
-		{
-			sourceName = inSourceName;
-			sourceValue = inSourceValue;
-			destName = inDestName;
-			destValue = inDestValue;
-		}
-		public string sourceName;
-		public string sourceValue;
-		public string destName;
-		public string destValue;
-	};
 
-	class DerivedMetadataMappings
-	{
-		public DerivedMetadataMappings()
-		{
-			entries = new List<DerivedMetadataEntry>();
-		}
-		public void ApplyMapping(CsvMetadata csvMetadata)
-		{
-			if (csvMetadata != null)
-			{
-				foreach (DerivedMetadataEntry entry in entries)
-				{
-					if (csvMetadata.Values.ContainsKey(entry.sourceName.ToLowerInvariant()))
-					{
-						if (csvMetadata.Values[entry.sourceName].ToLowerInvariant() == entry.sourceValue.ToLowerInvariant())
-						{
-							csvMetadata.Values.Add(entry.destName.ToLowerInvariant(), entry.destValue);
-						}
-					}
-				}
-			}
-		}
-		public List<DerivedMetadataEntry> entries;
-	}
-
-
-	class ReportXML
-	{
-		bool IsAbsolutePath(string path)
-		{
-			if (path.Length > 3 && path[1] == ':' && (path[2] == '\\' || path[2] == '/'))
-			{
-				return true;
-			}
-			return false;
-		}
-
-		public ReportXML(string graphXMLFilenameIn, string reportXMLFilenameIn, string baseXMLDirectoryOverride)
-		{
-			string location = System.Reflection.Assembly.GetEntryAssembly().Location.ToLower();
-			string baseDirectory = location.Replace("perfreporttool.exe", "");
-
-            // Check if this is a debug build, and redirect base dir to binaries if so
-            if ( baseDirectory.Contains("\\engine\\source\\programs\\") && baseDirectory.Contains("\\csvtools\\") && baseDirectory.Contains("\\bin\\debug\\"))
-            {
-                baseDirectory = baseDirectory.Replace("\\engine\\source\\programs\\", "\\engine\\binaries\\dotnet\\");
-                int csvToolsIndex = baseDirectory.LastIndexOf("\\csvtools\\");
-                baseDirectory = baseDirectory.Substring(0, csvToolsIndex + "\\csvtools\\".Length);
-            }
-
-			// Check if the base directory is being overridden
-			if (baseXMLDirectoryOverride.Length > 0)
-			{
-				if (IsAbsolutePath(baseXMLDirectoryOverride))
-				{
-					baseDirectory = baseXMLDirectoryOverride;
-				}
-				else
-				{
-					baseDirectory = Path.Combine(baseDirectory, baseXMLDirectoryOverride);
-				}
-			}
-			Console.Out.WriteLine("BaseDir: " + baseDirectory);
-
-			baseXmlDirectory = baseDirectory;
-
-			// Read the report type XML
-			reportTypeXmlFilename = Path.Combine(baseDirectory, "reportTypes.xml" );
-			if (reportXMLFilenameIn.Length > 0)
-			{
-				// Check if this is an absolute path
-				if (IsAbsolutePath(reportXMLFilenameIn))
-				{
-					reportTypeXmlFilename = reportXMLFilenameIn;
-				}
-				else
-				{
-					reportTypeXmlFilename = Path.Combine(baseDirectory, reportXMLFilenameIn);
-				}
-			}
-			XDocument reportTypesDoc = XDocument.Load(reportTypeXmlFilename);
-			rootElement = reportTypesDoc.Element("root");
-			if (rootElement == null)
-			{
-				throw new Exception("No root element found in report XML " + reportTypeXmlFilename);
-			}
-
-            reportTypesElement = rootElement.Element("reporttypes");
-            if (reportTypesElement == null)
-            {
-                throw new Exception("No reporttypes element found in report XML " + reportTypeXmlFilename);
-            }
-
-            // Read the graph XML
-            string graphsXMLFilename;
-            if (graphXMLFilenameIn.Length > 0)
-            {
-                if (IsAbsolutePath(graphXMLFilenameIn))
-                {
-                    graphsXMLFilename = graphXMLFilenameIn;
-                }
-                else
-                {
-                    graphsXMLFilename = Path.Combine( baseDirectory, graphXMLFilenameIn );
-                }
-            }
-            else
-            {
-                graphsXMLFilename = reportTypesElement.GetSafeAttibute<string>("reportGraphsFile");
-                if (graphsXMLFilename != null)
-                {
-                    graphsXMLFilename = Path.GetDirectoryName(reportTypeXmlFilename) + "\\" + graphsXMLFilename;
-                }
-                else
-                {
-                    graphsXMLFilename = Path.Combine( baseDirectory, "reportGraphs.xml" );
-                }
-
-            }
-
-
-
-
-			XDocument reportGraphsDoc = XDocument.Load(graphsXMLFilename);
-			graphGroupsElement = reportGraphsDoc.Element("graphGroups");
-
-			// Read the base settings - all other settings will inherit from this
-			GraphSettings baseSettings = new GraphSettings(graphGroupsElement.Element("baseSettings"));
-			if (reportTypesElement == null)
-			{
-				throw new Exception("No baseSettings element found in graph XML " + graphsXMLFilename);
-			}
-
-			graphs = new Dictionary<string, GraphSettings>();
-			foreach (XElement graphGroupElement in graphGroupsElement.Elements())
-			{
-				if (graphGroupElement.Name == "graphGroup")
-				{
-					// Create the base settings
-					XElement settingsElement = graphGroupElement.Element("baseSettings");
-					GraphSettings groupSettings = new GraphSettings(settingsElement);
-					groupSettings.InheritFrom(baseSettings);
-					foreach (XElement graphElement in graphGroupElement.Elements())
-					{
-						if (graphElement.Name == "graph")
-						{
-							string title = graphElement.Attribute("title").Value.ToLower();
-							GraphSettings graphSettings = new GraphSettings(graphElement);
-							graphSettings.InheritFrom(groupSettings);
-							graphs.Add(title, graphSettings);
-						}
-					}
-				}
-			}
-
-
-
-			// Read the display name mapping
-			statDisplayNameMapping = new Dictionary<string, string>();
-			XElement displayNameElement = rootElement.Element("statDisplayNameMappings");
-			if (displayNameElement != null)
-			{
-				foreach (XElement mapping in displayNameElement.Elements("mapping"))
-				{
-					string statName = mapping.GetSafeAttibute<string>("statName");
-					string displayName = mapping.GetSafeAttibute<string>("displayName");
-					if (statName != null && displayName != null)
-					{
-						statDisplayNameMapping.Add(statName.ToLower(), displayName);
-					}
-				}
-			}
-
-			XElement summaryTableColumnInfoListEl = rootElement.Element("summaryTableColumnFormatInfo");
-			if (summaryTableColumnInfoListEl != null)
-			{
-				columnFormatInfoList = new SummaryTableColumnFormatInfoCollection(summaryTableColumnInfoListEl);
-			}			
-
-			// Read the derived metadata mappings
-			derivedMetadataMappings = new DerivedMetadataMappings();
-			XElement derivedMetadataMappingsElement = rootElement.Element("derivedMetadataMappings");
-			if (derivedMetadataMappingsElement != null)
-			{
-				foreach (XElement mapping in derivedMetadataMappingsElement.Elements("mapping"))
-				{
-					string sourceName = mapping.GetSafeAttibute<string>("sourceName");
-					string sourceValue = mapping.GetSafeAttibute<string>("sourceValue");
-					string destName = mapping.GetSafeAttibute<string>("destName");
-					string destValue = mapping.GetSafeAttibute<string>("destValue");
-					if (sourceName == null || sourceValue == null || destName == null || destValue == null)
-					{
-						throw new Exception("Derivedmetadata mapping is missing a required attribute!\nRequired attributes: sourceName, sourceValue, destName, destValue.\nXML: "+mapping.ToString());
-					}
-					derivedMetadataMappings.entries.Add(new DerivedMetadataEntry(sourceName, sourceValue, destName, destValue));
-				}
-			}
-
-			// Read events to strip
-			XElement eventsToStripEl = rootElement.Element("csvEventsToStrip");
-            if (eventsToStripEl != null)
-            {
-                csvEventsToStrip = new List<CsvEventStripInfo>();
-                foreach (XElement eventPair in eventsToStripEl.Elements("eventPair"))
-                {
-                    CsvEventStripInfo eventInfo = new CsvEventStripInfo();
-                    eventInfo.beginName = eventPair.GetSafeAttibute<string>("begin");
-                    eventInfo.endName = eventPair.GetSafeAttibute<string>("end");
-
-                    if (eventInfo.beginName == null && eventInfo.endName == null )
-                    {
-                        throw new Exception("eventPair with no begin or end attribute found! Need to have one or the other.");
-                    }
-                    csvEventsToStrip.Add(eventInfo);
-                }
-            }
-
-			summaryTablesElement = rootElement.Element("summaryTables");
-			if (summaryTablesElement != null)
-			{
-				summaryTables = new Dictionary<string, SummaryTableInfo>();
-				foreach (XElement summaryElement in summaryTablesElement.Elements("summaryTable"))
-				{
-					SummaryTableInfo table = new SummaryTableInfo(summaryElement);
-					summaryTables.Add(summaryElement.Attribute("name").Value.ToLower(), table);
-				}
-			}
-
-			// Add any shared summaries
-			XElement sharedSummariesElement = rootElement.Element("sharedSummaries");
-			sharedSummaries = new Dictionary<string, XElement>();
-			if (sharedSummariesElement != null)
-			{
-				foreach (XElement summaryElement in sharedSummariesElement.Elements("summary"))
-				{
-					sharedSummaries.Add(summaryElement.Attribute("refName").Value, summaryElement);
-				}
-			}
-
-		}
-
-		public ReportTypeInfo GetReportTypeInfo(string reportType, CachedCsvFile csvFile, bool bBulkMode, bool forceReportType )
-		{
-			ReportTypeInfo reportTypeInfo = null;
-			if (reportType == "")
-			{
-				// Attempt to determine the report type automatically based on the stats
-				foreach (XElement element in reportTypesElement.Elements("reporttype"))
-				{
-					if (IsReportTypeXMLCompatibleWithStats(element, csvFile.dummyCsvStats))
-					{
-						reportTypeInfo = new ReportTypeInfo(element, sharedSummaries, baseXmlDirectory);
-						break;
-					}
-				}
-				if (reportTypeInfo == null)
-				{
-					throw new Exception("Compatible report type for CSV "+csvFile.filename+" could not be found in" + reportTypeXmlFilename);
-				}
-			}
-			else
-			{
-				XElement foundReportTypeElement = null;
-				foreach (XElement element in reportTypesElement.Elements("reporttype"))
-				{
-					if (element.Attribute("name").Value.ToLower() == reportType)
-					{
-						foundReportTypeElement = element;
-					}
-				}
-				if (foundReportTypeElement == null)
-				{
-					throw new Exception("Report type " + reportType + " not found in " + reportTypeXmlFilename);
-				}
-
-                if (!IsReportTypeXMLCompatibleWithStats(foundReportTypeElement, csvFile.dummyCsvStats))
-                {
-                    if (forceReportType)
-                    {
-                        Console.Out.WriteLine("Report type " + reportType + " is not compatible with CSV " + csvFile.filename + ", but using it anyway");
-                    }
-                    else
-                    {
-                        throw new Exception("Report type " + reportType + " is not compatible with CSV " + csvFile.filename);
-                    }
-                }
-                reportTypeInfo = new ReportTypeInfo(foundReportTypeElement, sharedSummaries, baseXmlDirectory);
-            }
-
-            // Load the graphs
-            foreach (ReportGraph graph in reportTypeInfo.graphs)
-			{
-				string key = graph.title.ToLower();
-				if (graphs.ContainsKey(key))
-				{
-					graph.settings = graphs[key];
-				}
-				else
-				{
-					throw new Exception("Graph with title \"" + graph.title + "\" was not found in graphs XML");
-				}
-			}
-
-			foreach (Summary summary in reportTypeInfo.summaries)
-			{
-				summary.PostInit(reportTypeInfo, csvFile.dummyCsvStats);
-			}
-			return reportTypeInfo;
-		} 
-
-		bool IsReportTypeXMLCompatibleWithStats(XElement reportTypeElement, CsvStats csvStats)
-		{
-			XAttribute nameAt = reportTypeElement.Attribute("name");
-			if (nameAt == null)
-			{
-				return false;
-			}
-			string reportTypeName = nameAt.Value;
-
-			XElement autoDetectionEl = reportTypeElement.Element("autodetection");
-			if (autoDetectionEl == null)
-			{
-				return false;
-			}
-			XAttribute requiredStatsAt = autoDetectionEl.Attribute("requiredstats");
-			if (requiredStatsAt != null)
-			{
-				string[] requiredStats = requiredStatsAt.Value.Split(',');
-				foreach (string stat in requiredStats)
-				{
-					if (csvStats.GetStatsMatchingString(stat).Count == 0)
-					{
-						return false;
-					}
-				}
-			}
-
-			foreach (XElement requiredMetadataEl in autoDetectionEl.Elements("requiredmetadata"))
-			{
-				XAttribute keyAt = requiredMetadataEl.Attribute("key");
-				if (keyAt == null)
-				{
-					throw new Exception("Report type " + reportTypeName + " has no 'key' attribute!");
-				}
-				XAttribute allowedValuesAt = requiredMetadataEl.Attribute("allowedValues");
-				if (allowedValuesAt == null)
-				{
-					throw new Exception("Report type " + reportTypeName + " has no 'allowedValues' attribute!");
-				}
-
-				if (csvStats.metaData == null)
-				{
-					// There was required metadata, but the CSV doesn't have any
-					return false;
-				}
-
-				bool ignoreIfKeyNotFound = requiredMetadataEl.GetSafeAttibute("ignoreIfKeyNotFound", true);
-				bool stopIfKeyFound = requiredMetadataEl.GetSafeAttibute("stopIfKeyFound", false);
-
-				string key = keyAt.Value.ToLower();
-				if (csvStats.metaData.Values.ContainsKey(key))
-				{
-					string value = csvStats.metaData.Values[key].ToLower();
-					string[] allowedValues = allowedValuesAt.Value.ToString().ToLower().Split(',');
-					if (!allowedValues.Contains(value))
-					{
-						return false;
-					}
-					if (stopIfKeyFound)
-					{
-						break;
-					}
-				}
-				else if (ignoreIfKeyNotFound == false)
-				{
-					return false;
-				}
-			}
-
-            //Console.Out.WriteLine("Autodetected report type: " + reportTypeName);
-
-			return true;
-        } 
-
-
-        public Dictionary<string, string> GetDisplayNameMapping() { return statDisplayNameMapping; }
-
-		public SummaryTableInfo GetSummaryTable(string name)
-		{
-			if (summaryTables.ContainsKey(name.ToLower()))
-			{
-				return summaryTables[name.ToLower()];
-			}
-			else
-			{
-				throw new Exception("Requested summary table type '" + name + "' was not found in <summaryTables>");
-			}
-		}
-
-        public List<CsvEventStripInfo> GetCsvEventsToStrip()
-        {
-            return csvEventsToStrip;
-        }
-
-		public void ApplyDerivedMetadata(CsvMetadata csvMetadata)
-		{
-			derivedMetadataMappings.ApplyMapping(csvMetadata);
-		}
-
-
-        Dictionary<string, SummaryTableInfo> summaryTables;
-
-		XElement reportTypesElement;
-		XElement rootElement;
-		XElement graphGroupsElement;
-		XElement summaryTablesElement;
-		Dictionary<string,XElement> sharedSummaries;
-		Dictionary<string, GraphSettings> graphs;
-		Dictionary<string, string> statDisplayNameMapping;
-		public SummaryTableColumnFormatInfoCollection columnFormatInfoList;
-		string baseXmlDirectory;
-
-		List<CsvEventStripInfo> csvEventsToStrip;
-        string reportTypeXmlFilename;
-		public DerivedMetadataMappings derivedMetadataMappings;
-	}
-
-	class ReportTypeParams
-	{
-		public string reportTypeOverride;
-		public bool forceReportType;
-	}
-
-	class CsvFileCache
-    {
-        public CsvFileCache( 
-			string[] inCsvFilenames, 
-			int inLookaheadCount, 
-			int inThreadCount, 
-			bool inUseCacheFiles, 
-			QueryExpression inMetadataQuery, 
-			ReportXML inReportXml, 
-			ReportTypeParams inReportTypeParams, 
-			bool inBulkMode, 
-			bool inSummaryTableCacheOnlyMode,
-			bool inSummaryTableCacheUseOnlyCsvID,
-			bool inRemoveDuplicates, 
-			bool inRequireMetadata,
-			string inSummaryTableCacheDir,
-			bool inListFilesMode)
-        {
-            csvFileInfos = new CsvFileInfo[inCsvFilenames.Length];
-            for (int i = 0; i < inCsvFilenames.Length; i++)
-            {
-                csvFileInfos[i] = new CsvFileInfo(inCsvFilenames[i]);
-            }
-            fileCache = this;
-            writeIndex = 0;
-			useCacheFiles = inUseCacheFiles;
-            readIndex = 0;
-            lookaheadCount = inLookaheadCount;
-            countFreedSinceLastGC = 0;
-			reportXml = inReportXml;
-			bulkMode = inBulkMode;
-			summaryTableCacheOnlyMode = inSummaryTableCacheOnlyMode;
-			summaryTableCacheUseOnlyCsvID = inSummaryTableCacheUseOnlyCsvID;
-			metadataQuery = inMetadataQuery;
-			derivedMetadataMappings = inReportXml.derivedMetadataMappings;
-			summaryTableCacheDir = inSummaryTableCacheDir;
-			reportTypeParams = inReportTypeParams;
-			bRemoveDuplicates = inRemoveDuplicates;
-			bRequireMetadata = inRequireMetadata;
-			bListFilesMode = inListFilesMode;
-
-			csvIdToFilename = new Dictionary<string, string>();
-			outCsvQueue = new BlockingCollection<CsvFileInfo>(lookaheadCount);
-
-			// Kick off the workers (must be done last)
-			if (inLookaheadCount > 0)
-            {
-				Console.WriteLine("Kicking off " + inThreadCount + " precache threads with lookahead " + lookaheadCount);
-                precacheThreads = new Thread[inThreadCount];
-                precacheJobs = new ThreadStart[inThreadCount];
-                for (int i = 0; i < precacheThreads.Length; i++)
-                { 
-                    precacheJobs[i] = new ThreadStart(PrecacheThreadRun);
-                    precacheThreads[i] = new Thread(precacheJobs[i]);
-                    precacheThreads[i].Start();
-                }
-            }
-		}
-
-        public CachedCsvFile GetNextCachedCsvFile()
-        {
-            CachedCsvFile file = null;
-            if ( readIndex >= csvFileInfos.Length)
-            {
-                // We're done
-                return null;
-            }
-
-			// Find the next valid fileinfo
-			if (precacheThreads == null)
-            {
-				CsvFileInfo fileInfo = csvFileInfos[readIndex];
-				file = new CachedCsvFile(fileInfo.filename, useCacheFiles, derivedMetadataMappings);
-				if (!file.DoesMetadataMatchQuery(metadataQuery)) // TODO do we need to check this here and in the thread?
-				{
-					return null;
-				}
-			}
-			else
-            {
-				CsvFileInfo fileInfo = outCsvQueue.Take();
-				file = fileInfo.cachedFile;
-				fileInfo.cachedFile = null;
-				countFreedSinceLastGC++;
-				// Periodically GC
-				if (countFreedSinceLastGC > 16 && summaryTableCacheOnlyMode == false)
-				{
-					GC.Collect();
-					GC.WaitForPendingFinalizers();
-					countFreedSinceLastGC = 0;
-				}
-				if (!fileInfo.isValid)
-				{
-					file = null;
-				}
-			}
-			readIndex++;
-            return file;
-        }
-
-        static void PrecacheThreadRun()
-        {
-            fileCache.ThreadRun();
-        }
-
-        void ThreadRun()
-        {
-            int threadWriteIndex = 0;
-            while (true)
-            {
-                threadWriteIndex = Interlocked.Increment(ref writeIndex)-1;
-                if ( threadWriteIndex >= csvFileInfos.Length )
-                {
-                    // We're done
-                    break;
-                }
-                CsvFileInfo fileInfo = csvFileInfos[threadWriteIndex];
-
-                // Process the file
-				CachedCsvFile file = new CachedCsvFile(fileInfo.filename, useCacheFiles, derivedMetadataMappings);
-
-				bool bProcessThisFile = true;
-				if (bRequireMetadata && file.metadata == null)
-				{
-					Console.WriteLine("CSV has no metadata. Skipping: " + fileInfo.filename);
-					bProcessThisFile = false;
-				}
-				else if (bRemoveDuplicates && file.metadata != null && file.metadata.Values.ContainsKey("csvid"))
-				{
-					string csvId = file.metadata.Values["csvid"];
-					lock (csvIdToFilename)
-					{
-						if ( csvIdToFilename.ContainsKey(csvId) )
-						{
-							Console.WriteLine("Duplicate CSV found: " + fileInfo.filename);
-							Console.WriteLine("   First version   : " + csvIdToFilename[csvId]);
-							bProcessThisFile = false;
-							duplicateCount++;
-						}
-						else
-						{
-							csvIdToFilename[csvId] = fileInfo.filename;
-						}
-					}
-				}
-
-				if ( bProcessThisFile && file.DoesMetadataMatchQuery(metadataQuery) )
-				{
-					if (bListFilesMode)
-					{
-						Console.WriteLine("File: " + fileInfo.filename);
-					}
-					else if (summaryTableCacheOnlyMode)
-					{
-						// Assume the report type is valid if we have cached metadata, since we don't actually need the reporttype in metadata only mode
-						if (file.cachedSummaryTableRowData != null)
-						{
-							// Just read the full cached metadata
-							file.cachedSummaryTableRowData = SummaryTableRowData.TryReadFromCacheFile(fileInfo.filename);
-							fileInfo.isValid = true;
-						}
-					}
-					else
-					{ 
-						file.reportTypeInfo = GetCsvReportTypeInfo(file, bulkMode);
-						if (file.reportTypeInfo != null)
-						{
-							string reportTypeHash = summaryTableCacheUseOnlyCsvID ? null : file.reportTypeInfo.GetSummaryTableCacheID();
-							file.ComputeSummaryTableCacheId(reportTypeHash);
-							if (summaryTableCacheDir != null)
-							{
-								// If a summary metadata cache is specified, try reading from it instead of reading the whole CSV
-								// Note that this will be disabled if we're not in bulk mode
-								file.cachedSummaryTableRowData = SummaryTableRowData.TryReadFromCache(summaryTableCacheDir, file.summaryTableCacheId);
-								if (file.cachedSummaryTableRowData == null)
-								{
-									Console.WriteLine("Failed to read summary metadata from cache for CSV: " + fileInfo.filename);
-								}
-							}
-							if (file.cachedSummaryTableRowData == null)
-							{
-								file.PrepareCsvData();
-							}
-
-							// Only read the full file data if the metadata matches
-							fileInfo.isValid = true;
-						}
-					}
-				}
-				fileInfo.cachedFile = file;
-				// TODO: only add valid files to the queue? Need alternative stopping condition
-				outCsvQueue.Add(fileInfo);
-			}
-        }
-
-		ReportTypeInfo GetCsvReportTypeInfo(CachedCsvFile csvFile, bool bBulkMode)
-		{
-			try
-			{
-				return reportXml.GetReportTypeInfo(reportTypeParams.reportTypeOverride, csvFile, bBulkMode, reportTypeParams.forceReportType);
-			}
-			catch (Exception e)
-			{
-				if (bBulkMode)
-				{
-					Console.Out.WriteLine("[ERROR] : " + e.Message);
-					return null;
-				}
-				else
-				{
-					// If we're not in bulk mode, exceptions are fatal
-					throw e;
-				}
-			}
-		}
-
-
-		ThreadStart[] precacheJobs;
-        Thread [] precacheThreads;
-        int writeIndex;
-        int readIndex;
-        int lookaheadCount;
-        int countFreedSinceLastGC;
-		bool useCacheFiles;
-		bool bulkMode;
-		bool bRequireMetadata;
-		bool summaryTableCacheOnlyMode;
-		bool summaryTableCacheUseOnlyCsvID;
-		QueryExpression metadataQuery;
-		string summaryTableCacheDir;
-		ReportXML reportXml;
-		ReportTypeParams reportTypeParams;
-		DerivedMetadataMappings derivedMetadataMappings;
-
-		Dictionary<string, string> csvIdToFilename;
-		public int duplicateCount=0;
-		bool bRemoveDuplicates;
-		bool bListFilesMode;
-
-		static CsvFileCache fileCache;
-		BlockingCollection<CsvFileInfo> outCsvQueue;
-
-        class CsvFileInfo
-        {
-            public CsvFileInfo(string inFilename)
-            {
-                filename = inFilename;
-                cachedFile = null;
-				isValid = false;
-            }
-            public CachedCsvFile cachedFile;
-            public string filename;
-			public bool isValid;
-        }
-        CsvFileInfo[] csvFileInfos;
-
-    }
 }
 
diff --git a/Engine/Source/Programs/CSVTools/PerfReportTool/PerfreportTool.csproj b/Engine/Source/Programs/CSVTools/PerfReportTool/PerfreportTool.csproj
--- a/Engine/Source/Programs/CSVTools/PerfReportTool/PerfreportTool.csproj
+++ b/Engine/Source/Programs/CSVTools/PerfReportTool/PerfreportTool.csproj
@@ -42,9 +42,23 @@
     <Reference Include="System.Xml" />
   </ItemGroup>
   <ItemGroup>
+    <Compile Include="CsvFileCache..cs" />
     <Compile Include="PerfReportTool.cs" />
-    <Compile Include="PerfSummaries.cs" />
+    <Compile Include="SummaryTypes.cs" />
     <Compile Include="Properties\AssemblyInfo.cs" />
+    <Compile Include="ReportGraph.cs" />
+    <Compile Include="ReportType.cs" />
+    <Compile Include="ReportXml.cs" />
+    <Compile Include="Summaries\BoundedStatValuesSummary.cs" />
+    <Compile Include="Summaries\EventSummary.cs" />
+    <Compile Include="Summaries\ExtraLinksSummary.cs" />
+    <Compile Include="Summaries\FPSChartSummary.cs" />
+    <Compile Include="Summaries\HistogramSummary.cs" />
+    <Compile Include="Summaries\HitchSummary.cs" />
+    <Compile Include="Summaries\MapOverlaySummary.cs" />
+    <Compile Include="Summaries\PeakSummary.cs" />
+    <Compile Include="Summary.cs" />
+    <Compile Include="SummaryTable.cs" />
   </ItemGroup>
   <ItemGroup>
     <None Include="App.config" />
diff --git a/Engine/Source/Programs/CSVTools/PerfReportTool/ReportGraph.cs b/Engine/Source/Programs/CSVTools/PerfReportTool/ReportGraph.cs
new file mode 100644
--- /dev/null
+++ b/Engine/Source/Programs/CSVTools/PerfReportTool/ReportGraph.cs
@@ -0,0 +1,269 @@
+// Copyright Epic Games, Inc. All Rights Reserved.
+
+using System;
+using System.Xml.Linq;
+
+namespace PerfReportTool
+{
+
+	class OptionalString
+	{
+		public OptionalString(string valueIn)
+		{
+			value = valueIn;
+			isSet = true;
+		}
+		public OptionalString()
+		{
+			isSet = false;
+		}
+		public OptionalString(XElement element, string Name, bool IsElement = false)
+		{
+			isSet = false;
+			if (IsElement)
+			{
+				XElement child = element.Element(Name);
+				if (child != null)
+				{
+					value = child.Value;
+					isSet = true;
+				}
+			}
+			else
+			{
+				XAttribute child = element.Attribute(Name);
+				if (child != null)
+				{
+					value = child.Value;
+					isSet = true;
+				}
+			}
+		}
+
+		public void InheritFrom(OptionalString baseVersion) { if (!isSet) { isSet = baseVersion.isSet; value = baseVersion.value; } }
+		public bool isSet;
+		public string value;
+	};
+
+	class OptionalBool
+	{
+		public OptionalBool(bool valueIn)
+		{
+			value = valueIn;
+			isSet = true;
+		}
+		public OptionalBool()
+		{
+			isSet = false;
+		}
+		public OptionalBool(XElement element, string AttributeName)
+		{
+			isSet = false;
+			try
+			{
+				if (element.Attribute(AttributeName) != null)
+				{
+					value = Convert.ToInt32(element.Attribute(AttributeName).Value) == 1;
+					isSet = true;
+				}
+			}
+			catch { }
+		}
+		public void InheritFrom(OptionalBool baseVersion) { if (!isSet) { isSet = baseVersion.isSet; value = baseVersion.value; } }
+
+		public bool isSet;
+		public bool value;
+	};
+
+	class OptionalInt
+	{
+		public OptionalInt(int valueIn)
+		{
+			value = valueIn;
+			isSet = true;
+		}
+		public OptionalInt()
+		{
+			isSet = false;
+		}
+		public OptionalInt(XElement element, string AttributeName)
+		{
+			isSet = false;
+			try
+			{
+				if (element.Attribute(AttributeName) != null)
+				{
+					value = Convert.ToInt32(element.Attribute(AttributeName).Value);
+					isSet = true;
+				}
+			}
+			catch { }
+		}
+		public void InheritFrom(OptionalInt baseVersion) { if (!isSet) { isSet = baseVersion.isSet; value = baseVersion.value; } }
+
+		public bool isSet;
+		public int value;
+	};
+
+	class OptionalDouble
+	{
+		public OptionalDouble(int valueIn)
+		{
+			value = valueIn;
+			isSet = true;
+		}
+		public OptionalDouble()
+		{
+			isSet = false;
+		}
+		public OptionalDouble(XElement element, string AttributeName)
+		{
+			isSet = false;
+			try
+			{
+				if (element.Attribute(AttributeName) != null)
+				{
+					value = Convert.ToDouble(element.Attribute(AttributeName).Value, System.Globalization.CultureInfo.InvariantCulture);
+					isSet = true;
+				}
+			}
+			catch { }
+		}
+
+		public void InheritFrom(OptionalDouble baseVersion) { if (!isSet) { isSet = baseVersion.isSet; value = baseVersion.value; } }
+
+		public bool isSet;
+		public double value;
+	};
+
+	static class OptionalHelper
+	{
+		public static string GetDoubleSetting(OptionalDouble setting, string cmdline)
+		{
+			return (setting.isSet ? (cmdline + setting.value.ToString()) : "");
+		}
+
+		public static string GetStringSetting(OptionalString setting, string cmdline)
+		{
+			return (setting.isSet ? (cmdline + setting.value) : "");
+		}
+	};
+
+	class ReportGraph
+	{
+		public ReportGraph(XElement element)
+		{
+			title = element.Attribute("title").Value;
+			budget = new OptionalDouble(element, "budget");
+			inSummary = element.GetSafeAttibute<bool>("inSummary", false);
+			isExternal = element.GetSafeAttibute<bool>("external", false);
+
+			minFilterStatValue = new OptionalDouble(element, "minFilterStatValue");
+		}
+		public string title;
+		public OptionalDouble budget;
+		public bool inSummary;
+		public bool isExternal;
+		public OptionalDouble minFilterStatValue;
+		public GraphSettings settings;
+	};
+
+	class GraphSettings
+	{
+		public GraphSettings(XElement element)
+		{
+			smooth = new OptionalBool(element, "smooth");
+			thickness = new OptionalDouble(element, "thickness");
+			miny = new OptionalDouble(element, "miny");
+			maxy = new OptionalDouble(element, "maxy");
+			threshold = new OptionalDouble(element, "threshold");
+			averageThreshold = new OptionalDouble(element, "averageThreshold");
+			minFilterStatValue = new OptionalDouble(element, "minFilterStatValue");
+			minFilterStatName = new OptionalString(element, "minFilterStatName");
+			smoothKernelPercent = new OptionalDouble(element, "smoothKernelPercent");
+			smoothKernelSize = new OptionalDouble(element, "smoothKernelSize");
+			compression = new OptionalDouble(element, "compression");
+			width = new OptionalInt(element, "width");
+			height = new OptionalInt(element, "height");
+			stacked = new OptionalBool(element, "stacked");
+			showAverages = new OptionalBool(element, "showAverages");
+			filterOutZeros = new OptionalBool(element, "filterOutZeros");
+			maxHierarchyDepth = new OptionalInt(element, "maxHierarchyDepth");
+			hideStatPrefix = new OptionalString(element, "hideStatPrefix");
+			mainStat = new OptionalString(element, "mainStat");
+			showEvents = new OptionalString(element, "showEvents");
+			requiresDetailedStats = new OptionalBool(element, "requiresDetailedStats");
+			ignoreStats = new OptionalString(element, "ignoreStats");
+
+			statString = new OptionalString(element, "statString", true);
+			additionalArgs = new OptionalString(element, "additionalArgs", true);
+			statMultiplier = new OptionalDouble(element, "statMultiplier");
+			legendAverageThreshold = new OptionalDouble(element, "legendAverageThreshold");
+			snapToPeaks = new OptionalBool(element, "snapToPeaks");
+			lineDecimalPlaces = new OptionalInt(element, "lineDecimalPlaces");
+		}
+		public void InheritFrom(GraphSettings baseSettings)
+		{
+			smooth.InheritFrom(baseSettings.smooth);
+			statString.InheritFrom(baseSettings.statString);
+			thickness.InheritFrom(baseSettings.thickness);
+			miny.InheritFrom(baseSettings.miny);
+			maxy.InheritFrom(baseSettings.maxy);
+			threshold.InheritFrom(baseSettings.threshold);
+			averageThreshold.InheritFrom(baseSettings.averageThreshold);
+			minFilterStatValue.InheritFrom(baseSettings.minFilterStatValue);
+			minFilterStatName.InheritFrom(baseSettings.minFilterStatName);
+			smoothKernelSize.InheritFrom(baseSettings.smoothKernelSize);
+			smoothKernelPercent.InheritFrom(baseSettings.smoothKernelPercent);
+			compression.InheritFrom(baseSettings.compression);
+			width.InheritFrom(baseSettings.width);
+			height.InheritFrom(baseSettings.height);
+			additionalArgs.InheritFrom(baseSettings.additionalArgs);
+			stacked.InheritFrom(baseSettings.stacked);
+			showAverages.InheritFrom(baseSettings.showAverages);
+			filterOutZeros.InheritFrom(baseSettings.filterOutZeros);
+			maxHierarchyDepth.InheritFrom(baseSettings.maxHierarchyDepth);
+			hideStatPrefix.InheritFrom(baseSettings.hideStatPrefix);
+			mainStat.InheritFrom(baseSettings.mainStat);
+			showEvents.InheritFrom(baseSettings.showEvents);
+			requiresDetailedStats.InheritFrom(baseSettings.requiresDetailedStats);
+			statMultiplier.InheritFrom(baseSettings.statMultiplier);
+			ignoreStats.InheritFrom(baseSettings.ignoreStats);
+			legendAverageThreshold.InheritFrom(baseSettings.legendAverageThreshold);
+			snapToPeaks.InheritFrom(baseSettings.snapToPeaks);
+			lineDecimalPlaces.InheritFrom(baseSettings.lineDecimalPlaces);
+
+		}
+		public OptionalBool smooth;
+		public OptionalString statString;
+		public OptionalDouble thickness;
+		public OptionalDouble miny;
+		public OptionalDouble maxy;
+		public OptionalDouble threshold;
+		public OptionalDouble averageThreshold;
+		public OptionalDouble minFilterStatValue;
+		public OptionalString minFilterStatName;
+		public OptionalDouble smoothKernelSize;
+		public OptionalDouble smoothKernelPercent;
+		public OptionalDouble compression;
+		public OptionalInt width;
+		public OptionalInt height;
+		public OptionalString additionalArgs;
+		public OptionalBool stacked;
+		public OptionalBool showAverages;
+		public OptionalBool filterOutZeros;
+		public OptionalInt maxHierarchyDepth;
+		public OptionalString hideStatPrefix;
+		public OptionalString mainStat;
+		public OptionalString showEvents;
+		public OptionalString ignoreStats;
+		public OptionalDouble statMultiplier;
+		public OptionalDouble legendAverageThreshold;
+
+		public OptionalBool requiresDetailedStats;
+		public OptionalBool snapToPeaks;
+		public OptionalInt lineDecimalPlaces;
+
+	};
+
+}
\ No newline at end of file
diff --git a/Engine/Source/Programs/CSVTools/PerfReportTool/ReportType.cs b/Engine/Source/Programs/CSVTools/PerfReportTool/ReportType.cs
new file mode 100644
--- /dev/null
+++ b/Engine/Source/Programs/CSVTools/PerfReportTool/ReportType.cs
@@ -0,0 +1,86 @@
+// Copyright Epic Games, Inc. All Rights Reserved.
+
+using System;
+using System.Collections.Generic;
+using System.Text;
+using System.Xml.Linq;
+using PerfSummaries;
+
+namespace PerfReportTool
+{
+	class ReportTypeParams
+	{
+		public string reportTypeOverride;
+		public bool forceReportType;
+	}
+
+	class ReportTypeInfo
+	{
+		public ReportTypeInfo(XElement element, Dictionary<string, XElement> sharedSummaries, string baseXmlDirectory)
+		{
+			graphs = new List<ReportGraph>();
+			summaries = new List<Summary>();
+			name = element.Attribute("name").Value;
+			title = element.Attribute("title").Value;
+			bStripEvents = element.GetSafeAttibute<bool>("stripEvents", true);
+
+			foreach (XElement child in element.Elements())
+			{
+				if (child.Name == "graph")
+				{
+					ReportGraph graph = new ReportGraph(child);
+					graphs.Add(graph);
+				}
+				else if (child.Name == "summary" || child.Name == "summaryRef")
+				{
+					XElement summaryElement = null;
+					if (child.Name == "summaryRef")
+					{
+						summaryElement = sharedSummaries[child.Attribute("name").Value];
+					}
+					else
+					{
+						summaryElement = child;
+					}
+					string summaryType = summaryElement.Attribute("type").Value;
+					summaries.Add(SummaryFactory.Create(summaryType, summaryElement, baseXmlDirectory));
+				}
+				else if (child.Name == "metadataToShow")
+				{
+					metadataToShowList = child.Value.Split(',');
+				}
+
+			}
+			ComputeSummaryTableCacheID();
+		}
+
+		public string GetSummaryTableCacheID()
+		{
+			return summaryTableCacheID;
+		}
+
+		private void ComputeSummaryTableCacheID()
+		{
+			StringBuilder sb = new StringBuilder();
+			sb.Append("NAME={" + name + "}\n");
+			sb.Append("TITLE={" + title + "}\n");
+			foreach (Summary summary in summaries)
+			{
+				sb.Append("SUMMARY={");
+				sb.Append("TYPE={" + summary.GetType().ToString() + "}");
+				sb.Append("UNSTRIPPED={" + summary.useUnstrippedCsvStats + "}");
+				sb.Append("STATS={" + string.Join(",", summary.stats) + "}");
+				sb.Append("}");
+			}
+			summaryTableCacheID = HashHelper.StringToHashStr(sb.ToString(), 16);
+		}
+
+		public List<ReportGraph> graphs;
+		public List<Summary> summaries;
+		public string name;
+		public string title;
+		public string[] metadataToShowList;
+		public string summaryTableCacheID;
+		public bool bStripEvents;
+	};
+}
\ No newline at end of file
diff --git a/Engine/Source/Programs/CSVTools/PerfReportTool/ReportXml.cs b/Engine/Source/Programs/CSVTools/PerfReportTool/ReportXml.cs
new file mode 100644
--- /dev/null
+++ b/Engine/Source/Programs/CSVTools/PerfReportTool/ReportXml.cs
@@ -0,0 +1,459 @@
+// Copyright Epic Games, Inc. All Rights Reserved.
+
+using System;
+using System.IO;
+using System.Collections.Generic;
+using System.Linq;
+using System.Xml.Linq;
+using CSVStats;
+using PerfSummaries;
+
+namespace PerfReportTool
+{
+	class CsvEventStripInfo
+	{
+		public string beginName;
+		public string endName;
+	};
+
+	class DerivedMetadataEntry
+	{
+		public DerivedMetadataEntry(string inSourceName, string inSourceValue, string inDestName, string inDestValue)
+		{
+			sourceName = inSourceName;
+			sourceValue = inSourceValue;
+			destName = inDestName;
+			destValue = inDestValue;
+		}
+		public string sourceName;
+		public string sourceValue;
+		public string destName;
+		public string destValue;
+	};
+
+	class DerivedMetadataMappings
+	{
+		public DerivedMetadataMappings()
+		{
+			entries = new List<DerivedMetadataEntry>();
+		}
+		public void ApplyMapping(CsvMetadata csvMetadata)
+		{
+			if (csvMetadata != null)
+			{
+				foreach (DerivedMetadataEntry entry in entries)
+				{
+					if (csvMetadata.Values.ContainsKey(entry.sourceName.ToLowerInvariant()))
+					{
+						if (csvMetadata.Values[entry.sourceName].ToLowerInvariant() == entry.sourceValue.ToLowerInvariant())
+						{
+							csvMetadata.Values.Add(entry.destName.ToLowerInvariant(), entry.destValue);
+						}
+					}
+				}
+			}
+		}
+		public List<DerivedMetadataEntry> entries;
+	}
+
+	class ReportXML
+	{
+		bool IsAbsolutePath(string path)
+		{
+			if (path.Length > 3 && path[1] == ':' && (path[2] == '\\' || path[2] == '/'))
+			{
+				return true;
+			}
+			return false;
+		}
+
+		public ReportXML(string graphXMLFilenameIn, string reportXMLFilenameIn, string baseXMLDirectoryOverride)
+		{
+			string location = System.Reflection.Assembly.GetEntryAssembly().Location.ToLower();
+			string baseDirectory = location.Replace("perfreporttool.exe", "");
+
+			// Check if this is a debug build, and redirect base dir to binaries if so
+			if (baseDirectory.Contains("\\engine\\source\\programs\\") && baseDirectory.Contains("\\csvtools\\") && baseDirectory.Contains("\\bin\\debug\\"))
+			{
+				baseDirectory = baseDirectory.Replace("\\engine\\source\\programs\\", "\\engine\\binaries\\dotnet\\");
+				int csvToolsIndex = baseDirectory.LastIndexOf("\\csvtools\\");
+				baseDirectory = baseDirectory.Substring(0, csvToolsIndex + "\\csvtools\\".Length);
+			}
+
+			// Check if the base directory is being overridden
+			if (baseXMLDirectoryOverride.Length > 0)
+			{
+				if (IsAbsolutePath(baseXMLDirectoryOverride))
+				{
+					baseDirectory = baseXMLDirectoryOverride;
+				}
+				else
+				{
+					baseDirectory = Path.Combine(baseDirectory, baseXMLDirectoryOverride);
+				}
+			}
+			Console.Out.WriteLine("BaseDir: " + baseDirectory);
+
+			baseXmlDirectory = baseDirectory;
+
+			// Read the report type XML
+			reportTypeXmlFilename = Path.Combine(baseDirectory, "reportTypes.xml");
+			if (reportXMLFilenameIn.Length > 0)
+			{
+				// Check if this is an absolute path
+				if (IsAbsolutePath(reportXMLFilenameIn))
+				{
+					reportTypeXmlFilename = reportXMLFilenameIn;
+				}
+				else
+				{
+					reportTypeXmlFilename = Path.Combine(baseDirectory, reportXMLFilenameIn);
+				}
+			}
+			XDocument reportTypesDoc = XDocument.Load(reportTypeXmlFilename);
+			rootElement = reportTypesDoc.Element("root");
+			if (rootElement == null)
+			{
+				throw new Exception("No root element found in report XML " + reportTypeXmlFilename);
+			}
+
+			reportTypesElement = rootElement.Element("reporttypes");
+			if (reportTypesElement == null)
+			{
+				throw new Exception("No reporttypes element found in report XML " + reportTypeXmlFilename);
+			}
+
+			// Read the graph XML
+			string graphsXMLFilename;
+			if (graphXMLFilenameIn.Length > 0)
+			{
+				if (IsAbsolutePath(graphXMLFilenameIn))
+				{
+					graphsXMLFilename = graphXMLFilenameIn;
+				}
+				else
+				{
+					graphsXMLFilename = Path.Combine(baseDirectory, graphXMLFilenameIn);
+				}
+			}
+			else
+			{
+				graphsXMLFilename = reportTypesElement.GetSafeAttibute<string>("reportGraphsFile");
+				if (graphsXMLFilename != null)
+				{
+					graphsXMLFilename = Path.GetDirectoryName(reportTypeXmlFilename) + "\\" + graphsXMLFilename;
+				}
+				else
+				{
+					graphsXMLFilename = Path.Combine(baseDirectory, "reportGraphs.xml");
+				}
+
+			}
+
+
+
+
+			XDocument reportGraphsDoc = XDocument.Load(graphsXMLFilename);
+			graphGroupsElement = reportGraphsDoc.Element("graphGroups");
+
+			// Read the base settings - all other settings will inherit from this
+			GraphSettings baseSettings = new GraphSettings(graphGroupsElement.Element("baseSettings"));
+			if (reportTypesElement == null)
+			{
+				throw new Exception("No baseSettings element found in graph XML " + graphsXMLFilename);
+			}
+
+			graphs = new Dictionary<string, GraphSettings>();
+			foreach (XElement graphGroupElement in graphGroupsElement.Elements())
+			{
+				if (graphGroupElement.Name == "graphGroup")
+				{
+					// Create the base settings
+					XElement settingsElement = graphGroupElement.Element("baseSettings");
+					GraphSettings groupSettings = new GraphSettings(settingsElement);
+					groupSettings.InheritFrom(baseSettings);
+					foreach (XElement graphElement in graphGroupElement.Elements())
+					{
+						if (graphElement.Name == "graph")
+						{
+							string title = graphElement.Attribute("title").Value.ToLower();
+							GraphSettings graphSettings = new GraphSettings(graphElement);
+							graphSettings.InheritFrom(groupSettings);
+							graphs.Add(title, graphSettings);
+						}
+					}
+				}
+			}
+
+
+
+			// Read the display name mapping
+			statDisplayNameMapping = new Dictionary<string, string>();
+			XElement displayNameElement = rootElement.Element("statDisplayNameMappings");
+			if (displayNameElement != null)
+			{
+				foreach (XElement mapping in displayNameElement.Elements("mapping"))
+				{
+					string statName = mapping.GetSafeAttibute<string>("statName");
+					string displayName = mapping.GetSafeAttibute<string>("displayName");
+					if (statName != null && displayName != null)
+					{
+						statDisplayNameMapping.Add(statName.ToLower(), displayName);
+					}
+				}
+			}
+
+			XElement summaryTableColumnInfoListEl = rootElement.Element("summaryTableColumnFormatInfo");
+			if (summaryTableColumnInfoListEl != null)
+			{
+				columnFormatInfoList = new SummaryTableColumnFormatInfoCollection(summaryTableColumnInfoListEl);
+			}
+
+			// Read the derived metadata mappings
+			derivedMetadataMappings = new DerivedMetadataMappings();
+			XElement derivedMetadataMappingsElement = rootElement.Element("derivedMetadataMappings");
+			if (derivedMetadataMappingsElement != null)
+			{
+				foreach (XElement mapping in derivedMetadataMappingsElement.Elements("mapping"))
+				{
+					string sourceName = mapping.GetSafeAttibute<string>("sourceName");
+					string sourceValue = mapping.GetSafeAttibute<string>("sourceValue");
+					string destName = mapping.GetSafeAttibute<string>("destName");
+					string destValue = mapping.GetSafeAttibute<string>("destValue");
+					if (sourceName == null || sourceValue == null || destName == null || destValue == null)
+					{
+						throw new Exception("Derivedmetadata mapping is missing a required attribute!\nRequired attributes: sourceName, sourceValue, destName, destValue.\nXML: " + mapping.ToString());
+					}
+					derivedMetadataMappings.entries.Add(new DerivedMetadataEntry(sourceName, sourceValue, destName, destValue));
+				}
+			}
+
+			// Read events to strip
+			XElement eventsToStripEl = rootElement.Element("csvEventsToStrip");
+			if (eventsToStripEl != null)
+			{
+				csvEventsToStrip = new List<CsvEventStripInfo>();
+				foreach (XElement eventPair in eventsToStripEl.Elements("eventPair"))
+				{
+					CsvEventStripInfo eventInfo = new CsvEventStripInfo();
+					eventInfo.beginName = eventPair.GetSafeAttibute<string>("begin");
+					eventInfo.endName = eventPair.GetSafeAttibute<string>("end");
+
+					if (eventInfo.beginName == null && eventInfo.endName == null)
+					{
+						throw new Exception("eventPair with no begin or end attribute found! Need to have one or the other.");
+					}
+					csvEventsToStrip.Add(eventInfo);
+				}
+			}
+
+			summaryTablesElement = rootElement.Element("summaryTables");
+			if (summaryTablesElement != null)
+			{
+				summaryTables = new Dictionary<string, SummaryTableInfo>();
+				foreach (XElement summaryElement in summaryTablesElement.Elements("summaryTable"))
+				{
+					SummaryTableInfo table = new SummaryTableInfo(summaryElement);
+					summaryTables.Add(summaryElement.Attribute("name").Value.ToLower(), table);
+				}
+			}
+
+			// Add any shared summaries
+			XElement sharedSummariesElement = rootElement.Element("sharedSummaries");
+			sharedSummaries = new Dictionary<string, XElement>();
+			if (sharedSummariesElement != null)
+			{
+				foreach (XElement summaryElement in sharedSummariesElement.Elements("summary"))
+				{
+					sharedSummaries.Add(summaryElement.Attribute("refName").Value, summaryElement);
+				}
+			}
+
+		}
+
+		public ReportTypeInfo GetReportTypeInfo(string reportType, CachedCsvFile csvFile, bool bBulkMode, bool forceReportType)
+		{
+			ReportTypeInfo reportTypeInfo = null;
+			if (reportType == "")
+			{
+				// Attempt to determine the report type automatically based on the stats
+				foreach (XElement element in reportTypesElement.Elements("reporttype"))
+				{
+					if (IsReportTypeXMLCompatibleWithStats(element, csvFile.dummyCsvStats))
+					{
+						reportTypeInfo = new ReportTypeInfo(element, sharedSummaries, baseXmlDirectory);
+						break;
+					}
+				}
+				if (reportTypeInfo == null)
+				{
+					throw new Exception("Compatible report type for CSV " + csvFile.filename + " could not be found in" + reportTypeXmlFilename);
+				}
+			}
+			else
+			{
+				XElement foundReportTypeElement = null;
+				foreach (XElement element in reportTypesElement.Elements("reporttype"))
+				{
+					if (element.Attribute("name").Value.ToLower() == reportType)
+					{
+						foundReportTypeElement = element;
+					}
+				}
+				if (foundReportTypeElement == null)
+				{
+					throw new Exception("Report type " + reportType + " not found in " + reportTypeXmlFilename);
+				}
+
+				if (!IsReportTypeXMLCompatibleWithStats(foundReportTypeElement, csvFile.dummyCsvStats))
+				{
+					if (forceReportType)
+					{
+						Console.Out.WriteLine("Report type " + reportType + " is not compatible with CSV " + csvFile.filename + ", but using it anyway");
+					}
+					else
+					{
+						throw new Exception("Report type " + reportType + " is not compatible with CSV " + csvFile.filename);
+					}
+				}
+				reportTypeInfo = new ReportTypeInfo(foundReportTypeElement, sharedSummaries, baseXmlDirectory);
+			}
+
+			// Load the graphs
+			foreach (ReportGraph graph in reportTypeInfo.graphs)
+			{
+				string key = graph.title.ToLower();
+				if (graphs.ContainsKey(key))
+				{
+					graph.settings = graphs[key];
+				}
+				else
+				{
+					throw new Exception("Graph with title \"" + graph.title + "\" was not found in graphs XML");
+				}
+			}
+
+			foreach (Summary summary in reportTypeInfo.summaries)
+			{
+				summary.PostInit(reportTypeInfo, csvFile.dummyCsvStats);
+			}
+			return reportTypeInfo;
+		}
+
+		bool IsReportTypeXMLCompatibleWithStats(XElement reportTypeElement, CsvStats csvStats)
+		{
+			XAttribute nameAt = reportTypeElement.Attribute("name");
+			if (nameAt == null)
+			{
+				return false;
+			}
+			string reportTypeName = nameAt.Value;
+
+			XElement autoDetectionEl = reportTypeElement.Element("autodetection");
+			if (autoDetectionEl == null)
+			{
+				return false;
+			}
+			XAttribute requiredStatsAt = autoDetectionEl.Attribute("requiredstats");
+			if (requiredStatsAt != null)
+			{
+				string[] requiredStats = requiredStatsAt.Value.Split(',');
+				foreach (string stat in requiredStats)
+				{
+					if (csvStats.GetStatsMatchingString(stat).Count == 0)
+					{
+						return false;
+					}
+				}
+			}
+
+			foreach (XElement requiredMetadataEl in autoDetectionEl.Elements("requiredmetadata"))
+			{
+				XAttribute keyAt = requiredMetadataEl.Attribute("key");
+				if (keyAt == null)
+				{
+					throw new Exception("Report type " + reportTypeName + " has no 'key' attribute!");
+				}
+				XAttribute allowedValuesAt = requiredMetadataEl.Attribute("allowedValues");
+				if (allowedValuesAt == null)
+				{
+					throw new Exception("Report type " + reportTypeName + " has no 'allowedValues' attribute!");
+				}
+
+				if (csvStats.metaData == null)
+				{
+					// There was required metadata, but the CSV doesn't have any
+					return false;
+				}
+
+				bool ignoreIfKeyNotFound = requiredMetadataEl.GetSafeAttibute("ignoreIfKeyNotFound", true);
+				bool stopIfKeyFound = requiredMetadataEl.GetSafeAttibute("stopIfKeyFound", false);
+
+				string key = keyAt.Value.ToLower();
+				if (csvStats.metaData.Values.ContainsKey(key))
+				{
+					string value = csvStats.metaData.Values[key].ToLower();
+					string[] allowedValues = allowedValuesAt.Value.ToString().ToLower().Split(',');
+					if (!allowedValues.Contains(value))
+					{
+						return false;
+					}
+					if (stopIfKeyFound)
+					{
+						break;
+					}
+				}
+				else if (ignoreIfKeyNotFound == false)
+				{
+					return false;
+				}
+			}
+
+			//Console.Out.WriteLine("Autodetected report type: " + reportTypeName);
+
+			return true;
+		}
+
+
+		public Dictionary<string, string> GetDisplayNameMapping() { return statDisplayNameMapping; }
+
+		public SummaryTableInfo GetSummaryTable(string name)
+		{
+			if (summaryTables.ContainsKey(name.ToLower()))
+			{
+				return summaryTables[name.ToLower()];
+			}
+			else
+			{
+				throw new Exception("Requested summary table type '" + name + "' was not found in <summaryTables>");
+			}
+		}
+
+		public List<CsvEventStripInfo> GetCsvEventsToStrip()
+		{
+			return csvEventsToStrip;
+		}
+
+		public void ApplyDerivedMetadata(CsvMetadata csvMetadata)
+		{
+			derivedMetadataMappings.ApplyMapping(csvMetadata);
+		}
+
+
+		Dictionary<string, SummaryTableInfo> summaryTables;
+
+		XElement reportTypesElement;
+		XElement rootElement;
+		XElement graphGroupsElement;
+		XElement summaryTablesElement;
+		Dictionary<string, XElement> sharedSummaries;
+		Dictionary<string, GraphSettings> graphs;
+		Dictionary<string, string> statDisplayNameMapping;
+		public SummaryTableColumnFormatInfoCollection columnFormatInfoList;
+		string baseXmlDirectory;
+
+		List<CsvEventStripInfo> csvEventsToStrip;
+		string reportTypeXmlFilename;
+		public DerivedMetadataMappings derivedMetadataMappings;
+	}
+}
\ No newline at end of file
diff --git a/Engine/Source/Programs/CSVTools/PerfReportTool/Summaries/BoundedStatValuesSummary.cs b/Engine/Source/Programs/CSVTools/PerfReportTool/Summaries/BoundedStatValuesSummary.cs
new file mode 100644
--- /dev/null
+++ b/Engine/Source/Programs/CSVTools/PerfReportTool/Summaries/BoundedStatValuesSummary.cs
@@ -0,0 +1,390 @@
+// Copyright Epic Games, Inc. All Rights Reserved.
+
+using System;
+using System.Collections.Generic;
+using System.Linq;
+using System.Xml.Linq;
+using PerfReportTool;
+using CSVStats;
+
+namespace PerfSummaries
+{
+	class BoundedStatValuesSummary : Summary
+	{
+		class Column
+		{
+			public string name;
+			public string formula;
+			public double value;
+			public string summaryStatName;
+			public string statName;
+			public string otherStatName;
+			public bool perSecond;
+			public bool filterOutZeros;
+			public bool applyEndOffset;
+			public double multiplier;
+			public double threshold;
+			public double frameExponent; // Exponent for relative frame time (0-1) in streamingstressmetric formula
+			public double statExponent; // Exponent for stat value in streamingstressmetric formula
+			public ColourThresholdList colourThresholdList;
+		};
+		public BoundedStatValuesSummary(XElement element, string baseXmlDirectory)
+		{
+			ReadStatsFromXML(element);
+			if (stats.Count != 0)
+			{
+				throw new Exception("<stats> element is not supported");
+			}
+
+			title = element.GetSafeAttibute("title", "Events");
+			beginEvent = element.GetSafeAttibute<string>("beginevent");
+			endEvent = element.GetSafeAttibute<string>("endevent");
+
+			endOffsetPercentage = 0.0;
+			XAttribute endOffsetAtt = element.Attribute("endoffsetpercent");
+			if (endOffsetAtt != null)
+			{
+				endOffsetPercentage = double.Parse(endOffsetAtt.Value);
+			}
+			columns = new List<Column>();
+
+			foreach (XElement columnEl in element.Elements("column"))
+			{
+				Column column = new Column();
+				double[] colourThresholds = ReadColourThresholdsXML(columnEl.Element("colourThresholds"));
+				if (colourThresholds != null)
+				{
+					column.colourThresholdList = new ColourThresholdList();
+					for (int i = 0; i < colourThresholds.Length; i++)
+					{
+						column.colourThresholdList.Add(new ThresholdInfo(colourThresholds[i], null));
+					}
+				}
+
+				XAttribute summaryStatNameAtt = columnEl.Attribute("summaryStatName");
+				if (summaryStatNameAtt != null)
+				{
+					column.summaryStatName = summaryStatNameAtt.Value;
+				}
+				column.statName = columnEl.Attribute("stat").Value.ToLower();
+				if (!stats.Contains(column.statName))
+				{
+					stats.Add(column.statName);
+				}
+				column.otherStatName = columnEl.GetSafeAttibute<string>("otherStat", "").ToLower();
+
+				column.name = columnEl.Attribute("name").Value;
+				column.formula = columnEl.Attribute("formula").Value.ToLower();
+				column.filterOutZeros = columnEl.GetSafeAttibute<bool>("filteroutzeros", false);
+				column.perSecond = columnEl.GetSafeAttibute<bool>("persecond", false);
+				column.multiplier = columnEl.GetSafeAttibute<double>("multiplier", 1.0);
+				column.threshold = columnEl.GetSafeAttibute<double>("threshold", 0.0);
+				column.applyEndOffset = columnEl.GetSafeAttibute<bool>("applyEndOffset", true);
+				column.frameExponent = columnEl.GetSafeAttibute<double>("frameExponent", 4.0);
+				column.statExponent = columnEl.GetSafeAttibute<double>("statExponent", 0.25);
+				columns.Add(column);
+			}
+		}
+
+		public BoundedStatValuesSummary() { }
+
+		public override string GetName() { return "boundedstatvalues"; }
+
+		public override void WriteSummaryData(System.IO.StreamWriter htmlFile, CsvStats csvStats, bool bWriteSummaryCsv, SummaryTableRowData rowData, string htmlFileName)
+		{
+			int startFrame = -1;
+			int endFrame = int.MaxValue;
+
+			// Find the start and end frames based on the events
+			if (beginEvent != null)
+			{
+				foreach (CsvEvent ev in csvStats.Events)
+				{
+					if (CsvStats.DoesSearchStringMatch(ev.Name, beginEvent))
+					{
+						startFrame = ev.Frame;
+						break;
+					}
+				}
+				if (startFrame == -1)
+				{
+					Console.WriteLine("BoundedStatValuesSummary: Begin event " + beginEvent + " was not found");
+					return;
+				}
+			}
+			if (endEvent != null)
+			{
+				foreach (CsvEvent ev in csvStats.Events)
+				{
+					if (CsvStats.DoesSearchStringMatch(ev.Name, endEvent))
+					{
+						endFrame = ev.Frame;
+						if (endFrame > startFrame)
+						{
+							break;
+						}
+					}
+				}
+				if (endFrame == int.MaxValue)
+				{
+					Console.WriteLine("BoundedStatValuesSummary: End event " + endEvent + " was not found");
+					return;
+				}
+			}
+			if (startFrame >= endFrame)
+			{
+				throw new Exception("BoundedStatValuesSummary: end event appeared before the start event");
+			}
+			endFrame = Math.Min(endFrame, csvStats.SampleCount - 1);
+			startFrame = Math.Max(startFrame, 0);
+
+			// Adjust the end frame based on the specified offset percentage, but cache the old value (some columns may need the unmodified one)
+			int endEventFrame = Math.Min(csvStats.SampleCount, endFrame + 1);
+			if (endOffsetPercentage > 0.0)
+			{
+				double multiplier = endOffsetPercentage / 100.0;
+				endFrame += (int)((double)(endFrame - startFrame) * multiplier);
+			}
+			endFrame = Math.Min(csvStats.SampleCount, endFrame + 1);
+			StatSamples frameTimeStat = csvStats.GetStat("frametime");
+			List<float> frameTimes = frameTimeStat.samples;
+
+			// Filter only columns with stats that exist in the CSV
+			List<Column> filteredColumns = new List<Column>();
+			foreach (Column col in columns)
+			{
+				if (csvStats.GetStat(col.statName) != null)
+				{
+					filteredColumns.Add(col);
+				}
+			}
+
+			// Nothing to report, so bail out!
+			if (filteredColumns.Count == 0)
+			{
+				return;
+			}
+
+			// Process the column values
+			foreach (Column col in filteredColumns)
+			{
+				List<float> statValues = csvStats.GetStat(col.statName).samples;
+				List<float> otherStatValues = csvStats.GetStat(col.otherStatName)?.samples;
+				double value = 0.0;
+				double totalFrameWeight = 0.0;
+				int colEndFrame = col.applyEndOffset ? endFrame : endEventFrame;
+
+				if (col.formula == "average")
+				{
+					for (int i = startFrame; i < colEndFrame; i++)
+					{
+						if (col.filterOutZeros == false || statValues[i] > 0)
+						{
+							value += statValues[i] * frameTimes[i];
+							totalFrameWeight += frameTimes[i];
+						}
+					}
+				}
+				else if (col.formula == "maximum")
+				{
+					for (int i = startFrame; i < colEndFrame; i++)
+					{
+						if (col.filterOutZeros == false || statValues[i] > 0)
+						{
+							value = statValues[i] > value ? statValues[i] : value;
+						}
+					}
+
+					totalFrameWeight = 1.0;
+				}
+				else if (col.formula == "percentoverthreshold")
+				{
+					for (int i = startFrame; i < colEndFrame; i++)
+					{
+						if (statValues[i] > col.threshold)
+						{
+							value += frameTimes[i];
+						}
+						totalFrameWeight += frameTimes[i];
+					}
+					value *= 100.0;
+				}
+				else if (col.formula == "percentunderthreshold")
+				{
+					for (int i = startFrame; i < colEndFrame; i++)
+					{
+						if (statValues[i] < col.threshold)
+						{
+							value += frameTimes[i];
+						}
+						totalFrameWeight += frameTimes[i];
+					}
+					value *= 100.0;
+				}
+				else if (col.formula == "sum")
+				{
+					for (int i = startFrame; i < colEndFrame; i++)
+					{
+						value += statValues[i];
+					}
+
+					if (col.perSecond)
+					{
+						double totalTimeMS = 0.0;
+						for (int i = startFrame; i < colEndFrame; i++)
+						{
+							if (col.filterOutZeros == false || statValues[i] > 0)
+							{
+								totalTimeMS += frameTimes[i];
+							}
+						}
+						value /= (totalTimeMS / 1000.0);
+					}
+					totalFrameWeight = 1.0;
+				}
+				else if (col.formula == "sumwhenotheroverthreshold")
+				{
+					for (int i = startFrame; i < colEndFrame; i++)
+					{
+						if (otherStatValues[i] > col.threshold)
+						{
+							value += statValues[i];
+						}
+					}
+
+					if (col.perSecond)
+					{
+						double totalTimeMS = 0.0;
+						for (int i = startFrame; i < colEndFrame; i++)
+						{
+							if ((col.filterOutZeros == false || statValues[i] > 0) && otherStatValues[i] > col.threshold)
+							{
+								totalTimeMS += frameTimes[i];
+							}
+						}
+						value /= (totalTimeMS / 1000.0);
+					}
+					totalFrameWeight = 1.0;
+				}
+				else if (col.formula == "sumwhenotherunderthreshold")
+				{
+					for (int i = startFrame; i < colEndFrame; i++)
+					{
+						if (otherStatValues[i] < col.threshold)
+						{
+							value += statValues[i];
+						}
+					}
+
+					if (col.perSecond)
+					{
+						double totalTimeMS = 0.0;
+						for (int i = startFrame; i < colEndFrame; i++)
+						{
+							if ((col.filterOutZeros == false || statValues[i] > 0) && otherStatValues[i] < col.threshold)
+							{
+								totalTimeMS += frameTimes[i];
+							}
+						}
+						value /= (totalTimeMS / 1000.0);
+					}
+					totalFrameWeight = 1.0;
+				}
+				else if (col.formula == "streamingstressmetric")
+				{
+					// Note: tInc is scaled such that it hits 1.0 on the event frame, regardless of the offset
+					double tInc = 1.0 / (double)(endEventFrame - startFrame);
+					double t = tInc * 0.5;
+					for (int i = startFrame; i < colEndFrame; i++)
+					{
+						if (col.filterOutZeros == false || statValues[i] > 0)
+						{
+							// Frame weighting is scaled to heavily favor final frames. Note that t can exceed 1 after the event frame if an offset percentage is specified, so we clamp it
+							double frameWeight = Math.Pow(Math.Min(t, 1.0), col.frameExponent) * frameTimes[i];
+
+							// If we're past the end event frame, apply a linear falloff to the weight
+							if (i >= endEventFrame)
+							{
+								double falloff = 1.0 - (double)(i - endEventFrame) / (colEndFrame - endEventFrame);
+								frameWeight *= falloff;
+							}
+
+							// The frame score takes into account the queue depth, but it's not massively significant
+							double frameScore = Math.Pow(statValues[i], col.statExponent);
+							value += frameScore * frameWeight;
+							totalFrameWeight += frameWeight;
+						}
+						t += tInc;
+					}
+				}
+				else if (col.formula == "ratio")
+				{
+					double numerator = 0.0;
+					for (int i = startFrame; i < colEndFrame; i++)
+					{
+						numerator += statValues[i];
+					}
+					double denominator = 0.0;
+					for (int i = startFrame; i < colEndFrame; i++)
+					{
+						denominator += otherStatValues[i];
+					}
+					value = numerator / denominator; // TODO: Does the rest of the pipeline handle +/-i infinity?
+					totalFrameWeight = 1.0;
+				}
+				else
+				{
+					throw new Exception("BoundedStatValuesSummary: unexpected formula " + col.formula);
+				}
+				value *= col.multiplier;
+				col.value = value / totalFrameWeight;
+			}
+
+			// Output HTML
+			if (htmlFile != null)
+			{
+				htmlFile.WriteLine("  <h2>" + title + "</h2>");
+				htmlFile.WriteLine("  <table border='0' style='width:1400'>");
+				htmlFile.WriteLine("  <tr>");
+				foreach (Column col in filteredColumns)
+				{
+					htmlFile.WriteLine("<th>" + col.name + "</th>");
+				}
+				htmlFile.WriteLine("  </tr>");
+				htmlFile.WriteLine("  <tr>");
+				foreach (Column col in filteredColumns)
+				{
+					string bgcolor = "'#ffffff'";
+					if (col.colourThresholdList != null)
+					{
+						bgcolor = col.colourThresholdList.GetColourForValue(col.value);
+					}
+					htmlFile.WriteLine("<td bgcolor=" + bgcolor + ">" + col.value.ToString("0.00") + "</td>");
+				}
+				htmlFile.WriteLine("  </tr>");
+				htmlFile.WriteLine("  </table>");
+			}
+
+			// Output summary table row data
+			if (rowData != null)
+			{
+				foreach (Column col in filteredColumns)
+				{
+					if (col.summaryStatName != null)
+					{
+						rowData.Add(SummaryTableElement.Type.SummaryTableMetric, col.summaryStatName, col.value, col.colourThresholdList);
+					}
+				}
+			}
+		}
+		public override void PostInit(ReportTypeInfo reportTypeInfo, CsvStats csvStats)
+		{
+		}
+		string title;
+		string beginEvent;
+		string endEvent;
+		double endOffsetPercentage;
+		List<Column> columns;
+	};
+
+}
\ No newline at end of file
diff --git a/Engine/Source/Programs/CSVTools/PerfReportTool/Summaries/EventSummary.cs b/Engine/Source/Programs/CSVTools/PerfReportTool/Summaries/EventSummary.cs
new file mode 100644
--- /dev/null
+++ b/Engine/Source/Programs/CSVTools/PerfReportTool/Summaries/EventSummary.cs
@@ -0,0 +1,94 @@
+// Copyright Epic Games, Inc. All Rights Reserved.
+
+using System;
+using System.Collections.Generic;
+using System.Linq;
+using System.Xml.Linq;
+using PerfReportTool;
+using CSVStats;
+
+namespace PerfSummaries
+{
+	class EventSummary : Summary
+	{
+		public EventSummary(XElement element, string baseXmlDirectory)
+		{
+			title = element.GetSafeAttibute("title", "Events");
+			summaryStatName = element.Attribute("summaryStatName").Value;
+			events = element.Element("events").Value.Split(',');
+			colourThresholds = ReadColourThresholdsXML(element.Element("colourThresholds"));
+		}
+
+		public EventSummary() { }
+
+		public override string GetName() { return "event"; }
+
+		public override void WriteSummaryData(System.IO.StreamWriter htmlFile, CsvStats csvStats, bool bWriteSummaryCsv, SummaryTableRowData rowData, string htmlFileName)
+		{
+			Dictionary<string, int> eventCountsDict = new Dictionary<string, int>();
+			int eventCount = 0;
+			foreach (CsvEvent ev in csvStats.Events)
+			{
+				foreach (string eventName in events)
+				{
+					if (CsvStats.DoesSearchStringMatch(ev.Name, eventName))
+					{
+						int len = eventName.Length;
+						if (eventName.EndsWith("*"))
+						{
+							len--;
+						}
+						string eventContent = ev.Name.Substring(len).Trim();
+						if (eventCountsDict.ContainsKey(eventContent))
+						{
+							eventCountsDict[eventContent]++;
+						}
+						else
+						{
+							eventCountsDict.Add(eventContent, 1);
+						}
+						eventCount++;
+					}
+				}
+			}
+
+			// Output HTML
+			if (htmlFile != null && eventCountsDict.Count > 0)
+			{
+				htmlFile.WriteLine("  <h2>" + title + "</h2>");
+				htmlFile.WriteLine("  <table border='0' style='width:1200'>");
+				htmlFile.WriteLine("  <tr><th>Name</th><th><b>Count</th></tr>");
+				foreach (KeyValuePair<string, int> pair in eventCountsDict.ToList())
+				{
+					htmlFile.WriteLine("  <tr><td>" + pair.Key + "</td><td>" + pair.Value + "</td></tr>");
+				}
+				htmlFile.WriteLine("  </table>");
+			}
+
+			// Output summary table row data
+			if (rowData != null)
+			{
+
+				ColourThresholdList thresholdList = null;
+
+				if (colourThresholds != null)
+				{
+					thresholdList = new ColourThresholdList();
+					for (int i = 0; i < colourThresholds.Length; i++)
+					{
+						thresholdList.Add(new ThresholdInfo(colourThresholds[i]));
+					}
+				}
+				rowData.Add(SummaryTableElement.Type.SummaryTableMetric, summaryStatName, (double)eventCount, thresholdList);
+			}
+		}
+		public override void PostInit(ReportTypeInfo reportTypeInfo, CsvStats csvStats)
+		{
+		}
+		string[] events;
+		double[] colourThresholds;
+		string title;
+		string summaryStatName;
+	};
+
+}
\ No newline at end of file
diff --git a/Engine/Source/Programs/CSVTools/PerfReportTool/Summaries/ExtraLinksSummary.cs b/Engine/Source/Programs/CSVTools/PerfReportTool/Summaries/ExtraLinksSummary.cs
new file mode 100644
--- /dev/null
+++ b/Engine/Source/Programs/CSVTools/PerfReportTool/Summaries/ExtraLinksSummary.cs
@@ -0,0 +1,104 @@
+// Copyright Epic Games, Inc. All Rights Reserved.
+
+using System;
+using System.Collections.Generic;
+using System.Linq;
+using System.Xml.Linq;
+using PerfReportTool;
+using System.IO;
+using CSVStats;
+
+namespace PerfSummaries
+{
+	class ExtraLinksSummary : Summary
+	{
+		class ExtraLink
+		{
+			public ExtraLink(string fileLine)
+			{
+				string[] Sections = fileLine.Split(',');
+				if (Sections.Length != 3)
+				{
+					throw new Exception("Bad links line format: " + fileLine);
+				}
+				LongName = Sections[0];
+				ShortName = Sections[1];
+				LinkURL = Sections[2];
+			}
+			public string GetLinkString(bool bUseLongName)
+			{
+				string Text = bUseLongName ? LongName : ShortName;
+				return "<a href='" + LinkURL + "'>" + Text + "</a>";
+			}
+			public string LongName;
+			public string ShortName;
+			public string LinkURL;
+		};
+
+		public ExtraLinksSummary(XElement element, string baseXmlDirectory)
+		{
+			title = "Links";
+			if (element != null)
+			{
+				title = element.GetSafeAttibute("title", title);
+			}
+		}
+		public ExtraLinksSummary() { }
+
+		public override string GetName() { return "extralinks"; }
+
+		public override void WriteSummaryData(System.IO.StreamWriter htmlFile, CsvStats csvStats, bool bWriteSummaryCsv, SummaryTableRowData rowData, string htmlFileName)
+		{
+			List<ExtraLink> links = new List<ExtraLink>();
+
+			string csvFilename = csvStats.metaData.GetValue("csvfilename", null);
+			if (csvFilename == null)
+			{
+				Console.WriteLine("Can't find CSV filename for ExtraLinks summary. Skipping");
+				return;
+			}
+
+			string linksFilename = csvFilename + ".links";
+			if (!File.Exists(linksFilename))
+			{
+				Console.WriteLine("Can't find file " + linksFilename + " for ExtraLinks summary. Skipping");
+				return;
+			}
+			string[] lines = File.ReadAllLines(linksFilename);
+			foreach (string line in lines)
+			{
+				links.Add(new ExtraLink(line));
+			}
+			if (links.Count == 0)
+			{
+				return;
+			}
+
+			// Output HTML
+			if (htmlFile != null)
+			{
+				htmlFile.WriteLine("  <h2>" + title + "</h2>");
+				htmlFile.WriteLine("  <ul>");
+				foreach (ExtraLink link in links)
+				{
+					htmlFile.WriteLine("  <li>" + link.GetLinkString(true) + "</li>");
+				}
+				htmlFile.WriteLine("  </ul>");
+			}
+
+			// Output summary row data
+			if (rowData != null)
+			{
+				foreach (ExtraLink link in links)
+				{
+					rowData.Add(SummaryTableElement.Type.SummaryTableMetric, link.LongName, link.GetLinkString(false), null);
+				}
+			}
+		}
+		public override void PostInit(ReportTypeInfo reportTypeInfo, CsvStats csvStats)
+		{
+		}
+		string title;
+	};
+
+}
\ No newline at end of file
diff --git a/Engine/Source/Programs/CSVTools/PerfReportTool/Summaries/FPSChartSummary.cs b/Engine/Source/Programs/CSVTools/PerfReportTool/Summaries/FPSChartSummary.cs
new file mode 100644
--- /dev/null
+++ b/Engine/Source/Programs/CSVTools/PerfReportTool/Summaries/FPSChartSummary.cs
@@ -0,0 +1,388 @@
+// Copyright Epic Games, Inc. All Rights Reserved.
+
+using System;
+using System.IO;
+using System.Collections.Generic;
+using System.Linq;
+using System.Xml.Linq;
+using PerfReportTool;
+using CSVStats;
+
+namespace PerfSummaries
+{
+	class FPSChartSummary : Summary
+	{
+		public FPSChartSummary(XElement element, string baseXmlDirectory)
+		{
+			ReadStatsFromXML(element);
+			fps = Convert.ToInt32(element.Attribute("fps").Value);
+			hitchThreshold = (float)Convert.ToDouble(element.Attribute("hitchThreshold").Value, System.Globalization.CultureInfo.InvariantCulture);
+			bUseEngineHitchMetric = element.GetSafeAttibute<bool>("useEngineHitchMetric", false);
+			if (bUseEngineHitchMetric)
+			{
+				engineHitchToNonHitchRatio = element.GetSafeAttibute<float>("engineHitchToNonHitchRatio", 1.5f);
+				engineMinTimeBetweenHitchesMs = element.GetSafeAttibute<float>("engineMinTimeBetweenHitchesMs", 200.0f);
+			}
+
+			bIgnoreHitchTimePercent = element.GetSafeAttibute<bool>("ignoreHitchTimePercent", false);
+			bIgnoreMVP = element.GetSafeAttibute<bool>("ignoreMVP", false);
+		}
+
+		public FPSChartSummary() { }
+
+		public override string GetName() { return "fpschart"; }
+
+		float GetEngineHitchToNonHitchRatio()
+		{
+			float MinimumRatio = 1.0f;
+			float targetFrameTime = 1000.0f / fps;
+			float MaximumRatio = hitchThreshold / targetFrameTime;
+
+			return Math.Min(Math.Max(engineHitchToNonHitchRatio, MinimumRatio), MaximumRatio);
+		}
+
+		struct FpsChartData
+		{
+			public float MVP;
+			public float HitchesPerMinute;
+			public float HitchTimePercent;
+			public int HitchCount;
+			public float TotalTimeSeconds;
+		};
+
+		FpsChartData ComputeFPSChartDataForFrames(List<float> frameTimes, bool skiplastFrame)
+		{
+			double totalFrametime = 0.0;
+			int hitchCount = 0;
+			double totalHitchTime = 0.0;
+
+			int frameCount = skiplastFrame ? frameTimes.Count - 1 : frameTimes.Count;
+
+			// Count hitches
+			if (bUseEngineHitchMetric)
+			{
+				// Minimum time passed before we'll record a new hitch
+				double CurrentTime = 0.0;
+				double LastHitchTime = float.MinValue;
+				double LastFrameTime = float.MinValue;
+				float HitchMultiplierAmount = GetEngineHitchToNonHitchRatio();
+
+				for (int i = 0; i < frameCount; i++)
+				{
+					float frametime = frameTimes[i];
+					// How long has it been since the last hitch we detected?
+					if (frametime >= hitchThreshold)
+					{
+						double TimeSinceLastHitch = (CurrentTime - LastHitchTime);
+						if (TimeSinceLastHitch >= engineMinTimeBetweenHitchesMs)
+						{
+							// For the current frame to be considered a hitch, it must have run at least this many times slower than
+							// the previous frame
+
+							// If our frame time is much larger than our last frame time, we'll count this as a hitch!
+							if (frametime > (LastFrameTime * HitchMultiplierAmount))
+							{
+								LastHitchTime = CurrentTime;
+								hitchCount++;
+							}
+						}
+						totalHitchTime += frametime;
+					}
+					LastFrameTime = frametime;
+					CurrentTime += (double)frametime;
+				}
+				totalFrametime = CurrentTime;
+			}
+			else
+			{
+				for (int i = 0; i < frameCount; i++)
+				{
+					float frametime = frameTimes[i];
+					totalFrametime += frametime;
+					if (frametime >= hitchThreshold)
+					{
+						hitchCount++;
+					}
+				}
+			}
+			float TotalSeconds = (float)totalFrametime / 1000.0f;
+			float TotalMinutes = TotalSeconds / 60.0f;
+
+			FpsChartData outData = new FpsChartData();
+			outData.HitchCount = hitchCount;
+			outData.TotalTimeSeconds = TotalSeconds;
+			outData.HitchesPerMinute = (float)hitchCount / TotalMinutes;
+			outData.HitchTimePercent = (float)(totalHitchTime / totalFrametime) * 100.0f;
+
+			int TotalTargetFrames = (int)((double)fps * (TotalSeconds));
+			int MissedFrames = Math.Max(TotalTargetFrames - frameTimes.Count, 0);
+			outData.MVP = (((float)MissedFrames * 100.0f) / (float)TotalTargetFrames);
+			return outData;
+		}
+
+
+		public override void WriteSummaryData(System.IO.StreamWriter htmlFile, CsvStats csvStats, bool bWriteSummaryCsv, SummaryTableRowData rowData, string htmlFileName)
+		{
+			System.IO.StreamWriter statsCsvFile = null;
+			if (bWriteSummaryCsv)
+			{
+				string csvPath = Path.Combine(Path.GetDirectoryName(htmlFileName), "FrameStats_colored.csv");
+				statsCsvFile = new System.IO.StreamWriter(csvPath, false);
+			}
+			// Compute MVP30 and MVP60. Note: we ignore the last frame because fpscharts can hitch
+			List<float> frameTimes = csvStats.Stats["frametime"].samples;
+			FpsChartData fpsChartData = ComputeFPSChartDataForFrames(frameTimes, true);
+
+			// Write the averages
+			List<string> ColumnNames = new List<string>();
+			List<double> ColumnValues = new List<double>();
+			List<string> ColumnColors = new List<string>();
+			List<ColourThresholdList> ColumnColorThresholds = new List<ColourThresholdList>();
+
+			ColumnNames.Add("Total Time (s)");
+			ColumnColorThresholds.Add(new ColourThresholdList());
+			ColumnValues.Add(fpsChartData.TotalTimeSeconds);
+			ColumnColors.Add(ColourThresholdList.GetSafeColourForValue(ColumnColorThresholds.Last(), ColumnValues.Last()));
+
+			ColumnNames.Add("Hitches/Min");
+			ColumnColorThresholds.Add(GetStatColourThresholdList(ColumnNames.Last()));
+			ColumnValues.Add(fpsChartData.HitchesPerMinute);
+			ColumnColors.Add(ColourThresholdList.GetSafeColourForValue(ColumnColorThresholds.Last(), ColumnValues.Last()));
+
+			if (!bIgnoreHitchTimePercent)
+			{
+				ColumnNames.Add("HitchTimePercent");
+				ColumnColorThresholds.Add(GetStatColourThresholdList(ColumnNames.Last()));
+				ColumnValues.Add(fpsChartData.HitchTimePercent);
+				ColumnColors.Add(ColourThresholdList.GetSafeColourForValue(ColumnColorThresholds.Last(), ColumnValues.Last()));
+			}
+
+			if (!bIgnoreMVP)
+			{
+				ColumnNames.Add("MVP" + fps.ToString());
+				ColumnColorThresholds.Add(GetStatColourThresholdList(ColumnNames.Last()));
+				ColumnValues.Add(fpsChartData.MVP);
+				ColumnColors.Add(ColourThresholdList.GetSafeColourForValue(ColumnColorThresholds.Last(), ColumnValues.Last()));
+			}
+
+			List<bool> ColumnIsAvgValueList = new List<bool>();
+			for (int i = 0; i < ColumnNames.Count; i++)
+			{
+				ColumnIsAvgValueList.Add(false);
+			}
+			foreach (string statName in stats)
+			{
+				string[] StatTokens = statName.Split('(');
+
+				float value = 0;
+				string ValueType = " Avg";
+				bool bIsAvg = false;
+				if (!csvStats.Stats.ContainsKey(StatTokens[0].ToLower()))
+				{
+					continue;
+				}
+				if (StatTokens.Length > 1 && StatTokens[1].ToLower().Contains("min"))
+				{
+					value = csvStats.Stats[StatTokens[0].ToLower()].ComputeMinValue();
+					ValueType = " Min";
+				}
+				else if (StatTokens.Length > 1 && StatTokens[1].ToLower().Contains("max"))
+				{
+					value = csvStats.Stats[StatTokens[0].ToLower()].ComputeMaxValue();
+					ValueType = " Max";
+				}
+				else
+				{
+					value = csvStats.Stats[StatTokens[0].ToLower()].average;
+					bIsAvg = true;
+				}
+				ColumnIsAvgValueList.Add(bIsAvg);
+				ColumnNames.Add(StatTokens[0] + ValueType);
+				ColumnValues.Add(value);
+				ColumnColorThresholds.Add(GetStatColourThresholdList(statName));
+				ColumnColors.Add(ColourThresholdList.GetSafeColourForValue(ColumnColorThresholds.Last(), ColumnValues.Last()));
+			}
+
+			// Output summary table row data
+			if (rowData != null)
+			{
+				for (int i = 0; i < ColumnNames.Count; i++)
+				{
+					string columnName = ColumnNames[i];
+
+					// Output simply MVP to rowData instead of MVP30 etc
+					if (columnName.StartsWith("MVP"))
+					{
+						columnName = "MVP";
+					}
+					// Hide pre-existing stats with the same name
+					if (ColumnIsAvgValueList[i] && columnName.EndsWith(" Avg"))
+					{
+						string originalStatName = columnName.Substring(0, columnName.Length - 4).ToLower();
+						SummaryTableElement smv;
+						if (rowData.dict.TryGetValue(originalStatName, out smv))
+						{
+							if (smv.type == SummaryTableElement.Type.CsvStatAverage)
+							{
+								smv.SetFlag(SummaryTableElement.Flags.Hidden, true);
+							}
+						}
+					}
+					rowData.Add(SummaryTableElement.Type.SummaryTableMetric, columnName, ColumnValues[i], ColumnColorThresholds[i]);
+				}
+				rowData.Add(SummaryTableElement.Type.SummaryTableMetric, "TargetFPS", (double)fps);
+			}
+
+			// Output HTML
+			if (htmlFile != null)
+			{
+				string HeaderRow = "";
+				string ValueRow = "";
+				HeaderRow += "<th>Section Name</th>";
+				ValueRow += "<td>Entire Run</td>";
+				for (int i = 0; i < ColumnNames.Count; i++)
+				{
+					string columnName = ColumnNames[i];
+					if (columnName.ToLower().EndsWith("time"))
+					{
+						columnName += " (ms)";
+					}
+					HeaderRow += "<th>" + TableUtil.FormatStatName(columnName) + "</th>";
+					ValueRow += "<td bgcolor=" + ColumnColors[i] + ">" + ColumnValues[i].ToString("0.00") + "</td>";
+				}
+				htmlFile.WriteLine("  <h2>FPSChart</h2>");
+				htmlFile.WriteLine("<table border='0' style='width:400'>");
+				htmlFile.WriteLine("  <tr>" + HeaderRow + "</tr>");
+				htmlFile.WriteLine("  <tr>" + ValueRow + "</tr>");
+			}
+
+			// Output CSV
+			if (statsCsvFile != null)
+			{
+				statsCsvFile.Write("Section Name,");
+				statsCsvFile.WriteLine(string.Join(",", ColumnNames));
+
+				statsCsvFile.Write("Entire Run,");
+				statsCsvFile.WriteLine(string.Join(",", ColumnValues));
+
+				// Pass through color data as part of database-friendly stuff.
+				statsCsvFile.Write("Entire Run BGColors,");
+				statsCsvFile.WriteLine(string.Join(",", ColumnColors));
+			}
+
+			if (csvStats.Events.Count > 0)
+			{
+				// Per-event breakdown
+				foreach (CaptureRange CapRange in captures)
+				{
+					ColumnValues.Clear();
+					ColumnColors.Clear();
+					CaptureData CaptureFrameTimes = GetFramesForCapture(CapRange, frameTimes, csvStats.Events);
+
+					if (CaptureFrameTimes == null)
+					{
+						continue;
+					}
+					FpsChartData captureFpsChartData = ComputeFPSChartDataForFrames(CaptureFrameTimes.Frames, true);
+
+					if (captureFpsChartData.TotalTimeSeconds == 0.0f)
+					{
+						continue;
+					}
+
+					ColumnValues.Add(captureFpsChartData.TotalTimeSeconds);
+					ColumnColors.Add("\'#ffffff\'");
+
+					ColumnValues.Add(captureFpsChartData.HitchesPerMinute);
+					ColumnColors.Add(GetStatThresholdColour("Hitches/Min", captureFpsChartData.HitchesPerMinute));
+
+					if (!bIgnoreHitchTimePercent)
+					{
+						ColumnValues.Add(captureFpsChartData.HitchTimePercent);
+						ColumnColors.Add(GetStatThresholdColour("HitchTimePercent", captureFpsChartData.HitchTimePercent));
+					}
+
+					if (!bIgnoreMVP)
+					{
+						ColumnValues.Add(captureFpsChartData.MVP);
+						ColumnColors.Add(GetStatThresholdColour("MVP" + fps.ToString(), captureFpsChartData.MVP));
+					}
+
+					foreach (string statName in stats)
+					{
+						string StatToCheck = statName.Split('(')[0];
+						if (!csvStats.Stats.ContainsKey(StatToCheck.ToLower()))
+						{
+							continue;
+						}
+
+						string[] StatTokens = statName.Split('(');
+
+						float value = 0;
+						if (StatTokens.Length > 1 && StatTokens[1].ToLower().Contains("min"))
+						{
+							value = csvStats.Stats[StatTokens[0].ToLower()].ComputeMinValue(CaptureFrameTimes.startIndex, CaptureFrameTimes.endIndex);
+						}
+						else if (StatTokens.Length > 1 && StatTokens[1].ToLower().Contains("max"))
+						{
+							value = csvStats.Stats[StatTokens[0].ToLower()].ComputeMaxValue(CaptureFrameTimes.startIndex, CaptureFrameTimes.endIndex);
+						}
+						else
+						{
+							value = csvStats.Stats[StatTokens[0].ToLower()].ComputeAverage(CaptureFrameTimes.startIndex, CaptureFrameTimes.endIndex);
+						}
+
+						ColumnValues.Add(value);
+						ColumnColors.Add(GetStatThresholdColour(statName, value));
+					}
+
+					// Output HTML
+					if (htmlFile != null)
+					{
+						string ValueRow = "";
+						ValueRow += "<td>" + CapRange.name + "</td>";
+						for (int i = 0; i < ColumnNames.Count; i++)
+						{
+							ValueRow += "<td bgcolor=" + ColumnColors[i] + ">" + ColumnValues[i].ToString("0.00") + "</td>";
+						}
+						htmlFile.WriteLine("  <tr>" + ValueRow + "</tr>");
+					}
+
+					// Output CSV
+					if (statsCsvFile != null)
+					{
+						statsCsvFile.Write(CapRange.name + ",");
+						statsCsvFile.WriteLine(string.Join(",", ColumnValues));
+
+						// Pass through color data as part of database-friendly stuff.
+						statsCsvFile.Write(CapRange.name + " colors,");
+						statsCsvFile.WriteLine(string.Join(",", ColumnColors));
+					}
+				}
+			}
+
+			if (htmlFile != null)
+			{
+				htmlFile.WriteLine("</table>");
+				htmlFile.WriteLine("<p style='font-size:8'>Engine hitch metric: " + (bUseEngineHitchMetric ? "enabled" : "disabled") + "</p>");
+			}
+
+			if (statsCsvFile != null)
+			{
+				statsCsvFile.Close();
+			}
+		}
+		public override void PostInit(ReportTypeInfo reportTypeInfo, CsvStats csvStats)
+		{
+		}
+
+		int fps;
+		float hitchThreshold;
+		bool bUseEngineHitchMetric;
+		bool bIgnoreHitchTimePercent;
+		bool bIgnoreMVP;
+		float engineHitchToNonHitchRatio;
+		float engineMinTimeBetweenHitchesMs;
+	};
+
+}
\ No newline at end of file
diff --git a/Engine/Source/Programs/CSVTools/PerfReportTool/Summaries/HistogramSummary.cs b/Engine/Source/Programs/CSVTools/PerfReportTool/Summaries/HistogramSummary.cs
new file mode 100644
--- /dev/null
+++ b/Engine/Source/Programs/CSVTools/PerfReportTool/Summaries/HistogramSummary.cs
@@ -0,0 +1,170 @@
+// Copyright Epic Games, Inc. All Rights Reserved.
+
+using System;
+using System.Collections.Generic;
+using System.Linq;
+using System.Xml.Linq;
+using PerfReportTool;
+using CSVStats;
+
+namespace PerfSummaries
+{
+	class HistogramSummary : Summary
+	{
+		public HistogramSummary(XElement element, string baseXmlDirectory)
+		{
+			ReadStatsFromXML(element);
+
+			ColourThresholds = ReadColourThresholdsXML(element.Element("colourThresholds"));
+
+			string[] histogramStrings = element.Element("histogramThresholds").Value.Split(',');
+			HistogramThresholds = new double[histogramStrings.Length];
+			for (int i = 0; i < histogramStrings.Length; i++)
+			{
+				HistogramThresholds[i] = Convert.ToDouble(histogramStrings[i], System.Globalization.CultureInfo.InvariantCulture);
+			}
+
+			string[] hitchThresholds = element.Element("hitchThresholds").Value.Split(',');
+			HitchThresholds = new double[hitchThresholds.Length];
+			for (int i = 0; i < hitchThresholds.Length; i++)
+			{
+				HitchThresholds[i] = Convert.ToDouble(hitchThresholds[i], System.Globalization.CultureInfo.InvariantCulture);
+			}
+
+			foreach (XElement child in element.Elements())
+			{
+				if (child.Name == "budgetOverride")
+				{
+					BudgetOverrideStatName = child.Attribute("stat").Value;
+					BudgetOverrideStatBudget = Convert.ToDouble(child.Attribute("budget").Value, System.Globalization.CultureInfo.InvariantCulture);
+				}
+			}
+		}
+
+		public HistogramSummary() { }
+
+		public override string GetName() { return "histogram"; }
+
+		public override void WriteSummaryData(System.IO.StreamWriter htmlFile, CsvStats csvStats, bool bWriteSummaryCsv, SummaryTableRowData metadata, string htmlFileName)
+		{
+			// Only HTML reporting is supported (does not output summary table row data)
+			if (htmlFile == null)
+			{
+				return;
+			}
+			// Write the averages
+			htmlFile.WriteLine("  <h2>Stat unit averages</h2>");
+			htmlFile.WriteLine("  <table border='0' style='width:400'>");
+			htmlFile.WriteLine("  <tr><td></td><th>ms</b></th>");
+			foreach (string stat in stats)
+			{
+				string StatToCheck = stat.Split('(')[0];
+				if (!csvStats.Stats.ContainsKey(StatToCheck.ToLower()))
+				{
+					continue;
+				}
+				float val = csvStats.Stats[StatToCheck.ToLower()].average;
+				string colour = ColourThresholdList.GetThresholdColour(val, ColourThresholds[0], ColourThresholds[1], ColourThresholds[2], ColourThresholds[3]);
+				htmlFile.WriteLine("  <tr><td><b>" + StatToCheck + "</b></td><td bgcolor=" + colour + ">" + val.ToString("0.00") + "</td></tr>");
+			}
+			htmlFile.WriteLine("  </table>");
+
+			// Hitches
+			double[] thresholds = HistogramThresholds;
+			htmlFile.WriteLine("  <h2>Frames in budget</h2>");
+			htmlFile.WriteLine("  <table border='0' style='width:800'>");
+
+			htmlFile.WriteLine("  <tr><td></td>");
+
+			// Display the override stat budget first
+			bool HasBudgetOverrideStat = false;
+			if (BudgetOverrideStatName != null)
+			{
+				htmlFile.WriteLine("  <td><b><=" + BudgetOverrideStatBudget.ToString("0") + "ms</b></td>");
+				HasBudgetOverrideStat = true;
+			}
+
+			foreach (float thresh in thresholds)
+			{
+				htmlFile.WriteLine("  <td><b><=" + thresh.ToString("0") + "ms</b></td>");
+			}
+			htmlFile.WriteLine("  </tr>");
+
+			foreach (string unitStat in stats)
+			{
+				string StatToCheck = unitStat.Split('(')[0];
+
+				htmlFile.WriteLine("  <tr><td><b>" + StatToCheck + "</b></td>");
+				int thresholdIndex = 0;
+
+				// Display the render thread budget column (don't display the other stats)
+				if (HasBudgetOverrideStat)
+				{
+					if (StatToCheck.ToLower() == BudgetOverrideStatName)
+					{
+						float pc = csvStats.GetStat(StatToCheck.ToLower()).GetRatioOfFramesInBudget((float)BudgetOverrideStatBudget) * 100.0f;
+						string colour = ColourThresholdList.GetThresholdColour(pc, 50.0f, 65.0f, 80.0f, 100.0f);
+						htmlFile.WriteLine("  <td bgcolor=" + colour + ">" + pc.ToString("0.00") + "%</td>");
+					}
+					else
+					{
+						htmlFile.WriteLine("  <td></td>");
+					}
+				}
+
+				foreach (float thresh in thresholds)
+				{
+					float threshold = (float)thresholds[thresholdIndex];
+					float pc = csvStats.GetStat(StatToCheck.ToLower()).GetRatioOfFramesInBudget(threshold) * 100.0f;
+					string colour = ColourThresholdList.GetThresholdColour(pc, 50.0f, 65.0f, 80.0f, 100.0f);
+					htmlFile.WriteLine("  <td bgcolor=" + colour + ">" + pc.ToString("0.00") + "%</td>");
+					thresholdIndex++;
+				}
+				htmlFile.WriteLine("  </tr>");
+			}
+			htmlFile.WriteLine("  </table>");
+
+
+			// Hitches
+			htmlFile.WriteLine("  <h2>Hitches - Overall</h2>");
+			htmlFile.WriteLine("  <table border='0' style='width:800'>");
+			htmlFile.WriteLine("  <tr><td></td>");
+
+			foreach (float thresh in HitchThresholds)
+			{
+				htmlFile.WriteLine("  <td><b> >" + thresh.ToString("0") + "ms</b></td>");
+			}
+
+			htmlFile.WriteLine("  </tr>");
+
+			foreach (string unitStat in stats)
+			{
+				string StatToCheck = unitStat.Split('(')[0];
+				htmlFile.WriteLine("  <tr><td><b>" + unitStat + "</b></td>");
+				int thresholdIndex = 0;
+
+				foreach (float threshold in HitchThresholds)
+				{
+					float count = (float)csvStats.GetStat(StatToCheck.ToLower()).GetCountOfFramesOverBudget(threshold);
+					int numSamples = csvStats.GetStat(StatToCheck.ToLower()).GetNumSamples();
+					// if we have 20k frames in a typical flythrough then 20 frames would be red
+					float redThresholdFor50ms = (float)numSamples / 500.0f;
+					float redThreshold = (redThresholdFor50ms * 50.0f) / threshold; // Adjust the colour threshold based on the current threshold
+					string colour = ColourThresholdList.GetThresholdColour(count, redThreshold, redThreshold * 0.66, redThreshold * 0.33, 0.0f);
+					htmlFile.WriteLine("  <td bgcolor=" + colour + ">" + count.ToString("0") + "</td>");
+					thresholdIndex++;
+				}
+				htmlFile.WriteLine("  </tr>");
+			}
+			htmlFile.WriteLine("  </table>");
+
+		}
+
+		public double[] ColourThresholds;
+		public double[] HistogramThresholds;
+		public double[] HitchThresholds;
+		public string BudgetOverrideStatName;
+		public double BudgetOverrideStatBudget;
+	};
+
+}
\ No newline at end of file
diff --git a/Engine/Source/Programs/CSVTools/PerfReportTool/Summaries/HitchSummary.cs b/Engine/Source/Programs/CSVTools/PerfReportTool/Summaries/HitchSummary.cs
new file mode 100644
--- /dev/null
+++ b/Engine/Source/Programs/CSVTools/PerfReportTool/Summaries/HitchSummary.cs
@@ -0,0 +1,106 @@
+// Copyright Epic Games, Inc. All Rights Reserved.
+
+using System;
+using System.IO;
+using System.Collections.Generic;
+using System.Linq;
+using System.Xml.Linq;
+using PerfReportTool;
+using CSVStats;
+
+namespace PerfSummaries
+{
+	class HitchSummary : Summary
+	{
+		public HitchSummary(XElement element, string baseXmlDirectory)
+		{
+			ReadStatsFromXML(element);
+
+			string[] hitchThresholds = element.Element("hitchThresholds").Value.Split(',');
+			HitchThresholds = new double[hitchThresholds.Length];
+			for (int i = 0; i < hitchThresholds.Length; i++)
+			{
+				HitchThresholds[i] = Convert.ToDouble(hitchThresholds[i], System.Globalization.CultureInfo.InvariantCulture);
+			}
+		}
+		public HitchSummary() { }
+
+		public override string GetName() { return "hitches"; }
+
+		public override void WriteSummaryData(System.IO.StreamWriter htmlFile, CsvStats csvStats, bool bWriteSummaryCsv, SummaryTableRowData metadata, string htmlFileName)
+		{
+			// Only HTML reporting is supported (does not summary table row data)
+			if (htmlFile == null)
+			{
+				return;
+			}
+
+			htmlFile.WriteLine("  <h2>Hitches</h2>");
+			htmlFile.WriteLine("  <table border='0' style='width:800'>");
+			htmlFile.WriteLine("  <tr><td></td>");
+
+			StreamWriter statsCsvFile = null;
+			if (bWriteSummaryCsv)
+			{
+				string csvPath = Path.Combine(Path.GetDirectoryName(htmlFileName), "HitchStats.csv");
+				statsCsvFile = new System.IO.StreamWriter(csvPath, false);
+			}
+
+			List<string> Thresholds = new List<string>();
+			List<string> Hitches = new List<string>();
+			Thresholds.Add("Hitch Size");
+			foreach (float thresh in HitchThresholds)
+			{
+				htmlFile.WriteLine("  <th> >" + thresh.ToString("0") + "ms</b></td>");
+				Thresholds.Add(thresh.ToString("0"));
+			}
+			if (statsCsvFile != null)
+			{
+				statsCsvFile.WriteLine(string.Join(",", Thresholds));
+			}
+			htmlFile.WriteLine("  </tr>");
+
+			foreach (string unitStat in stats)
+			{
+				string StatToCheck = unitStat.Split('(')[0];
+				StatSamples statSample = csvStats.GetStat(StatToCheck.ToLower());
+				if (statSample == null)
+				{
+					continue;
+				}
+
+				Hitches.Clear();
+				htmlFile.WriteLine("  <tr><td><b>" + StatToCheck + "</b></td>");
+				Hitches.Add(StatToCheck);
+				int thresholdIndex = 0;
+
+				foreach (float threshold in HitchThresholds)
+				{
+					float count = (float)statSample.GetCountOfFramesOverBudget(threshold);
+					int numSamples = csvStats.GetStat(StatToCheck.ToLower()).GetNumSamples();
+					// if we have 20k frames in a typical flythrough then 20 frames would be red
+					float redThresholdFor50ms = (float)numSamples / 500.0f;
+					float redThreshold = (redThresholdFor50ms * 50.0f) / threshold; // Adjust the colour threshold based on the current threshold
+					string colour = ColourThresholdList.GetThresholdColour(count, redThreshold, redThreshold * 0.66, redThreshold * 0.33, 0.0f);
+					htmlFile.WriteLine("  <td bgcolor=" + colour + ">" + count.ToString("0") + "</td>");
+					Hitches.Add(count.ToString("0"));
+					thresholdIndex++;
+				}
+				if (statsCsvFile != null)
+				{
+					statsCsvFile.WriteLine(string.Join(",", Hitches));
+				}
+
+				htmlFile.WriteLine("  </tr>");
+			}
+			if (statsCsvFile != null)
+			{
+				statsCsvFile.Close();
+			}
+			htmlFile.WriteLine("  </table>");
+			htmlFile.WriteLine("<p style='font-size:8'>Note: Simplified hitch metric. All frames over threshold are counted" + "</p>");
+		}
+		public double[] HitchThresholds;
+	};
+
+}
\ No newline at end of file
diff --git a/Engine/Source/Programs/CSVTools/PerfReportTool/Summaries/MapOverlaySummary.cs b/Engine/Source/Programs/CSVTools/PerfReportTool/Summaries/MapOverlaySummary.cs
new file mode 100644
--- /dev/null
+++ b/Engine/Source/Programs/CSVTools/PerfReportTool/Summaries/MapOverlaySummary.cs
@@ -0,0 +1,275 @@
+// Copyright Epic Games, Inc. All Rights Reserved.
+
+using System;
+using System.Collections.Generic;
+using System.Linq;
+using System.Xml.Linq;
+using PerfReportTool;
+using CSVStats;
+
+namespace PerfSummaries
+{
+	class MapOverlaySummary : Summary
+	{
+		class MapOverlayEvent
+		{
+			public MapOverlayEvent(string inName)
+			{
+				name = inName;
+			}
+			public MapOverlayEvent(XElement element)
+			{
+			}
+			public string name;
+			public string summaryStatName;
+			public string shortName;
+			public string lineColor;
+		};
+
+		class MapOverlay
+		{
+			public MapOverlay(XElement element)
+			{
+				positionStatNames[0] = element.GetSafeAttibute<string>("xStat");
+				positionStatNames[1] = element.GetSafeAttibute<string>("yStat");
+				positionStatNames[2] = element.GetSafeAttibute<string>("zStat");
+				summaryStatNamePrefix = element.GetSafeAttibute<string>("summaryStatNamePrefix"); // unused!
+				lineColor = element.GetSafeAttibute<string>("lineColor", "#ffffff");
+				foreach (XElement eventEl in element.Elements("event"))
+				{
+					MapOverlayEvent ev = new MapOverlayEvent(eventEl.Attribute("name").Value);
+					ev.shortName = eventEl.GetSafeAttibute<string>("shortName");
+					ev.summaryStatName = eventEl.GetSafeAttibute<string>("summaryStatName"); // unused!
+					ev.lineColor = eventEl.GetSafeAttibute<string>("lineColor");
+					if (eventEl.GetSafeAttibute<bool>("isStartEvent", false))
+					{
+						if (startEvent != null)
+						{
+							throw new Exception("Can't have multiple start events!");
+						}
+						startEvent = ev;
+					}
+					events.Add(ev);
+				}
+
+			}
+			public string[] positionStatNames = new string[3];
+			public string summaryStatNamePrefix;
+			public MapOverlayEvent startEvent;
+			public string lineColor;
+			public List<MapOverlayEvent> events = new List<MapOverlayEvent>();
+		}
+
+		public MapOverlaySummary(XElement element, string baseXmlDirectory)
+		{
+			ReadStatsFromXML(element);
+			if (stats.Count != 0)
+			{
+				throw new Exception("<stats> element is not supported");
+			}
+
+			sourceImagePath = element.GetSafeAttibute<string>("sourceImage");
+			if (baseXmlDirectory == null)
+			{
+				throw new Exception("BaseXmlDirectory not specified");
+			}
+			if (!System.IO.Path.IsPathRooted(sourceImagePath))
+			{
+				sourceImagePath = System.IO.Path.GetFullPath(System.IO.Path.Combine(baseXmlDirectory, sourceImagePath));
+			}
+
+			offsetX = element.GetSafeAttibute<float>("offsetX", 0.0f);
+			offsetY = element.GetSafeAttibute<float>("offsetY", 0.0f);
+			scale = element.GetSafeAttibute<float>("scale", 1.0f);
+			title = element.GetSafeAttibute("title", "Events");
+			destImageFilename = element.Attribute("destImage").Value;
+			imageWidth = element.GetSafeAttibute<float>("width", 250.0f);
+			imageHeight = element.GetSafeAttibute<float>("height", 250.0f);
+			framesPerLineSegment = element.GetSafeAttibute<int>("framesPerLineSegment", 5);
+			lineSplitDistanceThreshold = element.GetSafeAttibute<float>("lineSplitDistanceThreshold", float.MaxValue);
+
+			foreach (XElement overlayEl in element.Elements("overlay"))
+			{
+				MapOverlay overlay = new MapOverlay(overlayEl);
+				overlays.Add(overlay);
+				stats.Add(overlay.positionStatNames[0]);
+				stats.Add(overlay.positionStatNames[1]);
+				stats.Add(overlay.positionStatNames[2]);
+			}
+		}
+		public MapOverlaySummary() { }
+
+		public override string GetName() { return "mapoverlay"; }
+
+		int toSvgX(float worldX, float worldY)
+		{
+			float svgX = (worldY * scale + offsetX) * 0.5f + 0.5f;
+			svgX *= imageWidth;
+			return (int)(svgX + 0.5f);
+		}
+
+		int toSvgY(float worldX, float worldY)
+		{
+			float svgY = 1.0f - (worldX * scale + offsetY) * 0.5f - 0.5f;
+			svgY *= imageHeight;
+			return (int)(svgY + 0.5f);
+		}
+
+		public override void WriteSummaryData(System.IO.StreamWriter htmlFile, CsvStats csvStats, bool bWriteSummaryCsv, SummaryTableRowData rowData, string htmlFileName)
+		{
+			// Output HTML
+			if (htmlFile != null)
+			{
+				string outputDirectory = System.IO.Path.GetDirectoryName(System.IO.Path.GetFullPath(htmlFileName));
+				string outputMapFilename = System.IO.Path.Combine(outputDirectory, destImageFilename);
+
+				if (!System.IO.File.Exists(outputMapFilename))
+				{
+					System.IO.File.Copy(sourceImagePath, outputMapFilename);
+				}
+
+				// Check if the file exists in the output directory
+				htmlFile.WriteLine("  <h2>" + title + "</h2>");
+				htmlFile.WriteLine("<svg version='1.1' xmlns='http://www.w3.org/2000/svg' xmlns:xlink='http://www.w3.org/1999/xlink' width='" + imageWidth + "' height='" + imageHeight + "'>");
+				htmlFile.WriteLine("<image href='" + destImageFilename + "' width='" + imageWidth + "' height='" + imageHeight + "' />");
+
+				// Draw the overlays
+				foreach (MapOverlay overlay in overlays)
+				{
+					StatSamples xStat = csvStats.GetStat(overlay.positionStatNames[0]);
+					StatSamples yStat = csvStats.GetStat(overlay.positionStatNames[1]);
+
+					if (xStat == null || yStat == null)
+					{
+						continue;
+					}
+
+					// If a startevent is specified, update the start frame
+					int startFrame = 0;
+					if (overlay.startEvent != null)
+					{
+						foreach (CsvEvent ev in csvStats.Events)
+						{
+							if (CsvStats.DoesSearchStringMatch(ev.Name, overlay.startEvent.name))
+							{
+								startFrame = ev.Frame;
+								break;
+							}
+						}
+					}
+
+					// Make a mapping from frame to map indices
+					List<KeyValuePair<int, MapOverlayEvent>> frameEvents = new List<KeyValuePair<int, MapOverlayEvent>>();
+					foreach (MapOverlayEvent mapEvent in overlay.events)
+					{
+						foreach (CsvEvent ev in csvStats.Events)
+						{
+							if (CsvStats.DoesSearchStringMatch(ev.Name, mapEvent.name))
+							{
+								frameEvents.Add(new KeyValuePair<int, MapOverlayEvent>(ev.Frame, mapEvent));
+							}
+						}
+					}
+					frameEvents.Sort((pair0, pair1) => pair0.Key.CompareTo(pair1.Key));
+					int eventIndex = 0;
+
+					// Draw the lines
+					string currentLineColor = overlay.lineColor;
+					string lineStartTemplate = "<polyline style='fill:none;stroke-width:1.3;stroke:{LINECOLOUR}' points='";
+					htmlFile.Write(lineStartTemplate.Replace("{LINECOLOUR}", currentLineColor));
+					float adjustedLineSplitDistanceThreshold = lineSplitDistanceThreshold * framesPerLineSegment;
+					float oldx = 0;
+					float oldy = 0;
+					int lastFrameIndex = 0;
+					for (int i = startFrame; i < xStat.samples.Count; i += framesPerLineSegment)
+					{
+						float x = xStat.samples[i];
+						float y = yStat.samples[i];
+						string lineCoordsStr = toSvgX(x, y) + "," + toSvgY(x, y) + " ";
+
+						// Figure out which event we're up to so we can do color changes
+						bool restartLineStrip = false;
+						while (eventIndex < frameEvents.Count && lastFrameIndex < frameEvents[eventIndex].Key && i >= frameEvents[eventIndex].Key)
+						{
+							MapOverlayEvent mapEvent = frameEvents[eventIndex].Value;
+							string newLineColor = mapEvent.lineColor != null ? mapEvent.lineColor : overlay.lineColor;
+							// If we changed color, restart the line strip
+							if (newLineColor != currentLineColor)
+							{
+								currentLineColor = newLineColor;
+								restartLineStrip = true;
+							}
+							eventIndex++;
+						}
+
+						// If the distance between this point and the last is over the threshold, restart the line strip
+						float maxManhattanDist = Math.Max(Math.Abs(x - oldx), Math.Abs(y - oldy));
+						if (maxManhattanDist > adjustedLineSplitDistanceThreshold)
+						{
+							restartLineStrip = true;
+						}
+						else
+						{
+							htmlFile.Write(lineCoordsStr);
+						}
+
+						if (restartLineStrip)
+						{
+							htmlFile.WriteLine("'/>");
+							htmlFile.Write(lineStartTemplate.Replace("{LINECOLOUR}", currentLineColor));
+							htmlFile.Write(lineCoordsStr);
+						}
+						oldx = x;
+						oldy = y;
+						lastFrameIndex = i;
+					}
+					htmlFile.WriteLine("'/>");
+
+					// Plot the events 
+					float circleRadius = 3;
+					string eventColourString = "#ffffff";
+					foreach (MapOverlayEvent mapEvent in overlay.events)
+					{
+						foreach (CsvEvent ev in csvStats.Events)
+						{
+							if (CsvStats.DoesSearchStringMatch(ev.Name, mapEvent.name))
+							{
+								string eventText = mapEvent.shortName != null ? mapEvent.shortName : ev.Name;
+								float x = xStat.samples[ev.Frame];
+								float y = yStat.samples[ev.Frame];
+								int svgX = toSvgX(x, y);
+								int svgY = toSvgY(x, y);
+								htmlFile.Write("<circle cx='" + svgX + "' cy='" + svgY + "' r='" + circleRadius + "' fill='" + eventColourString + "' fill-opacity='1.0'/>");
+								htmlFile.WriteLine("<text x='" + (svgX + 5) + "' y='" + svgY + "' text-anchor='left' style='font-family: Verdana;fill: #ffffff; font-size: " + 9 + "px;'>" + eventText + "</text>");
+							}
+						}
+					}
+				}
+
+				//htmlFile.WriteLine("<text x='50%' y='" + (imageHeight * 0.05) + "' text-anchor='middle' style='font-family: Verdana;fill: #FFFFFF; stroke: #C0C0C0;  font-size: " + 20 + "px;'>" + title + "</text>");
+				htmlFile.WriteLine("</svg>");
+			}
+
+			// Output row data
+			if (rowData != null)
+			{
+			}
+		}
+		public override void PostInit(ReportTypeInfo reportTypeInfo, CsvStats csvStats)
+		{
+		}
+		string title;
+		string sourceImagePath;
+		float offsetX;
+		float offsetY;
+		float scale;
+		string destImageFilename;
+		float imageWidth;
+		float imageHeight;
+		float lineSplitDistanceThreshold;
+		int framesPerLineSegment;
+
+		List<MapOverlay> overlays = new List<MapOverlay>();
+	};
+
+}
\ No newline at end of file
diff --git a/Engine/Source/Programs/CSVTools/PerfReportTool/Summaries/PeakSummary.cs b/Engine/Source/Programs/CSVTools/PerfReportTool/Summaries/PeakSummary.cs
new file mode 100644
--- /dev/null
+++ b/Engine/Source/Programs/CSVTools/PerfReportTool/Summaries/PeakSummary.cs
@@ -0,0 +1,284 @@
+// Copyright Epic Games, Inc. All Rights Reserved.
+
+using System;
+using System.IO;
+using System.Collections.Generic;
+using System.Linq;
+using System.Xml.Linq;
+using PerfReportTool;
+using CSVStats;
+
+namespace PerfSummaries
+{
+	class PeakSummary : Summary
+	{
+		/*
+		  A peak summary displays a list of stats by their peak values. 
+		  The stat list comes from the graphs, where inSummary='1' specifies that a graph's stats should be
+		  included. Specifying a budget for the graph also assigns budgets to its stats. 
+		  A list of summarySection elements can be specified. This groups stats into separate sections, displayed
+		  at the top of the report. 
+		  Note: a stat will only appear in in one section. If there are multiple compatble sections, it will appear 
+		  in the first.
+
+			<summarySection title="Audio">
+				<statFilter>LLM/Audio/*</statFilter>
+			</summarySection>
+		*/
+		class PeakSummarySection
+		{
+			public PeakSummarySection(string inTitle, string inStatFilterStr)
+			{
+				title = inTitle;
+				statNamesFilter = inStatFilterStr.Split(',');
+			}
+			public PeakSummarySection(XElement element)
+			{
+				XElement statFilterElement = element.Element("statFilter");
+				title = element.Attribute("title").Value;
+
+				statNamesFilter = statFilterElement.Value.Split(',');
+			}
+			public bool StatMatchesSection(CsvStats csvStats, string statName)
+			{
+				if (statNameFilterDict == null)
+				{
+					statNameFilterDict = csvStats.GetStatNamesMatchingStringList_Dict(statNamesFilter);
+				}
+				return statNameFilterDict.ContainsKey(statName);
+			}
+			public void AddStat(PeakStatInfo statInfo)
+			{
+				stats.Add(statInfo);
+			}
+			string[] statNamesFilter;
+			Dictionary<string, bool> statNameFilterDict;
+
+			public List<PeakStatInfo> stats = new List<PeakStatInfo>();
+			public string title;
+		};
+
+		public PeakSummary(XElement element, string baseXmlDirectory)
+		{
+			//read the child elements (mostly for colourThresholds)
+			ReadStatsFromXML(element);
+			hideStatPrefix = element.GetSafeAttibute<string>("hideStatPrefix", "").ToLower();
+
+			foreach (XElement child in element.Elements())
+			{
+				if (child.Name == "summarySection")
+				{
+					peakSummarySections.Add(new PeakSummarySection(child));
+				}
+			}
+
+			// If we don't have any sections then add a default one
+			if (peakSummarySections.Count == 0)
+			{
+				PeakSummarySection section = new PeakSummarySection("Peaks", "*");
+				peakSummarySections.Add(section);
+			}
+		}
+
+		public PeakSummary() { }
+
+		public override string GetName() { return "peak"; }
+
+		void WriteStatSection(StreamWriter htmlFile, CsvStats csvStats, PeakSummarySection section, StreamWriter LLMCsvData, SummaryTableRowData summaryTableRowData)
+		{
+			// Here we are deciding which title we have and write it to the file.
+			htmlFile.WriteLine("<h3>" + section.title + "</h3>");
+			htmlFile.WriteLine("  <table border='0' style='width:400'>");
+
+			//Hard-coded start of the table.
+			htmlFile.WriteLine("    <tr><td style='width:200'></td><td style='width:75'><b>Average</b></td><td style='width:75'><b>Peak</b></td><td style='width:75'><b>Budget</b></td></tr>");
+
+			foreach (PeakStatInfo statInfo in section.stats)
+			{
+				// Do the calculations for the averages and peak, and then write it to the table along with the budget.
+				string statName = statInfo.name;
+				StatSamples csvStat = csvStats.Stats[statName.ToLower()];
+				double peak = (double)csvStat.ComputeMaxValue();
+				double average = (double)csvStat.average;
+
+				string peakColour = "#ffffff";
+				string averageColour = "#ffffff";
+				string budgetString = "";
+				ColourThresholdList colorThresholdList = new ColourThresholdList();
+
+				if (statInfo.budget.isSet)
+				{
+					double budget = statInfo.budget.value;
+					float redValue = (float)budget * 1.5f;
+					float orangeValue = (float)budget * 1.25f;
+					float yellowValue = (float)budget * 1.0f;
+					float greenValue = (float)budget * 0.9f;
+
+					colorThresholdList.Add(new ThresholdInfo(greenValue));
+					colorThresholdList.Add(new ThresholdInfo(yellowValue));
+					colorThresholdList.Add(new ThresholdInfo(orangeValue));
+					colorThresholdList.Add(new ThresholdInfo(redValue));
+
+					peakColour = colorThresholdList.GetColourForValue(peak);
+					averageColour = colorThresholdList.GetColourForValue(average);
+					budgetString = budget.ToString("0");
+				}
+				htmlFile.WriteLine("    <tr><td>" + statInfo.shortName + "</td><td bgcolor=" + averageColour + ">" + average.ToString("0") + "</td><td bgcolor=" + peakColour + ">" + peak.ToString("0") + "</td><td>" + budgetString + "</td></tr>");
+
+				// Pass through color data as part of database-friendly stuff.
+				if (LLMCsvData != null)
+				{
+					string csvStatName = statName.Replace('/', ' ').Replace("$32$", " ");
+					LLMCsvData.WriteLine(string.Format("{0},{1},{2},{3}", csvStatName, average.ToString("0"), peak.ToString("0"), budgetString, averageColour, peakColour));
+					LLMCsvData.WriteLine(string.Format("{0}_Colors,{1},{2},'#aaaaaa'", csvStatName, averageColour, peakColour));
+				}
+
+				if (summaryTableRowData != null)
+				{
+					SummaryTableElement smv;
+					// Hide duplicate CsvStatAverage stats
+					if (summaryTableRowData.dict.TryGetValue(statName.ToLower(), out smv))
+					{
+						if (smv.type == SummaryTableElement.Type.CsvStatAverage)
+						{
+							smv.SetFlag(SummaryTableElement.Flags.Hidden, true);
+						}
+					}
+					summaryTableRowData.Add(SummaryTableElement.Type.SummaryTableMetric, statInfo.shortName + " Avg", average, colorThresholdList);
+					summaryTableRowData.Add(SummaryTableElement.Type.SummaryTableMetric, statInfo.shortName + " Max", peak, colorThresholdList);
+				}
+
+			}
+			htmlFile.WriteLine("  </table>");
+		}
+
+		PeakSummarySection FindStatSection(CsvStats csvStats, string statName)
+		{
+			for (int i = 0; i < peakSummarySections.Count; i++)
+			{
+				if (peakSummarySections[i].StatMatchesSection(csvStats, statName))
+				{
+					return peakSummarySections[i];
+				}
+			}
+			return null;
+		}
+
+		public override void WriteSummaryData(System.IO.StreamWriter htmlFile, CsvStats csvStats, bool bWriteSummaryCsv, SummaryTableRowData summaryTableRowData, string htmlFileName)
+		{
+			// Only HTML reporting is supported (does not output summary table row data)
+			if (htmlFile == null)
+			{
+				return;
+			}
+
+			StreamWriter LLMCsvData = null;
+			if (bWriteSummaryCsv)
+			{
+				// FIXME: This summary type is not specific to LLM. Pass filename in!
+				string LLMCsvPath = Path.Combine(Path.GetDirectoryName(htmlFileName), "LLMStats_colored.csv");
+				LLMCsvData = new StreamWriter(LLMCsvPath);
+				LLMCsvData.WriteLine("Stat,Average,Peak,Budget");
+			}
+
+			// Add all stats to the appropriate sections
+			foreach (string stat in stats)
+			{
+				PeakSummarySection section = FindStatSection(csvStats, stat);
+				if (section != null)
+				{
+					PeakStatInfo statInfo = getOrAddStatInfo(stat);
+					section.AddStat(statInfo);
+				}
+			}
+
+			htmlFile.WriteLine("<h2>Peaks Summary</h2>");
+			foreach (PeakSummarySection section in peakSummarySections)
+			{
+				WriteStatSection(htmlFile, csvStats, section, LLMCsvData, summaryTableRowData);
+			}
+
+			if (LLMCsvData != null)
+			{
+				LLMCsvData.Close();
+			}
+		}
+
+
+
+		void AddStat(string statName, OptionalDouble budget)
+		{
+			stats.Add(statName);
+
+			PeakStatInfo info = getOrAddStatInfo(statName);
+			if (budget.isSet)
+			{
+				info.budget = budget;
+			}
+		}
+
+		public override void PostInit(ReportTypeInfo reportTypeInfo, CsvStats csvStats)
+		{
+			// Find the stats by spinning through the graphs in this reporttype
+			foreach (ReportGraph graph in reportTypeInfo.graphs)
+			{
+				if (graph.inSummary)
+				{
+					if (graph.settings.mainStat.isSet)
+					{
+						AddStat(graph.settings.mainStat.value, graph.budget);
+					}
+					if (graph.settings.statString.isSet)
+					{
+						string statString = graph.settings.statString.value;
+						string[] statNames = statString.Split(',');
+						statNames = csvStats.GetStatNamesMatchingStringList(statNames).ToArray();
+						foreach (string stat in statNames)
+						{
+							AddStat(stat, graph.budget);
+						}
+					}
+				}
+			}
+
+			base.PostInit(reportTypeInfo, csvStats);
+		}
+
+		List<PeakSummarySection> peakSummarySections = new List<PeakSummarySection>();
+
+		Dictionary<string, PeakStatInfo> statInfoLookup = new Dictionary<string, PeakStatInfo>();
+		class PeakStatInfo
+		{
+			public PeakStatInfo(string inName, string inShortName)
+			{
+				budget = new OptionalDouble();
+				name = inName;
+				shortName = inShortName;
+			}
+			public string name;
+			public string shortName;
+			public OptionalDouble budget;
+		};
+
+		PeakStatInfo getOrAddStatInfo(string statName)
+		{
+			if (statInfoLookup.ContainsKey(statName))
+			{
+				return statInfoLookup[statName];
+			}
+			// Find the best (longest) prefix which matches this stat, and strip it off
+			string shortStatName = statName;
+			if (hideStatPrefix.Length > 0 && statName.ToLower().StartsWith(hideStatPrefix))
+			{
+				shortStatName = statName.Substring(hideStatPrefix.Length);
+			}
+
+			PeakStatInfo statInfo = new PeakStatInfo(statName, shortStatName);
+			statInfoLookup.Add(statName, statInfo);
+			return statInfo;
+		}
+
+		string hideStatPrefix;
+	};
+
+}
\ No newline at end of file
diff --git a/Engine/Source/Programs/CSVTools/PerfReportTool/Summary.cs b/Engine/Source/Programs/CSVTools/PerfReportTool/Summary.cs
new file mode 100644
--- /dev/null
+++ b/Engine/Source/Programs/CSVTools/PerfReportTool/Summary.cs
@@ -0,0 +1,251 @@
+﻿// Copyright Epic Games, Inc. All Rights Reserved.
+
+using System;
+using System.Collections.Generic;
+using System.Linq;
+using System.Xml.Linq;
+using CSVStats;
+using PerfReportTool;
+using System.Reflection;
+
+namespace PerfSummaries
+{
+	class SummaryFactory
+	{
+		public static void Init()
+		{
+			summaryNameLookup = new Dictionary<string, Type>();
+			Assembly assembly = typeof(SummaryFactory).GetTypeInfo().Assembly;
+			foreach (Type type in assembly.GetTypes())
+			{
+				if (type.IsSubclassOf(typeof(Summary)))
+				{
+					Summary instance = (Summary)Activator.CreateInstance(type);
+					summaryNameLookup.Add(instance.GetName(), type);
+				}
+			}
+		}
+
+		public static Summary Create(string summaryTypeName, XElement summaryXmlElement, string baseXmlDirectory)
+		{
+			if ( summaryNameLookup.TryGetValue(summaryTypeName, out System.Type summaryType ) )
+			{
+				object[] constructArgs = new object[2] { summaryXmlElement, baseXmlDirectory };
+				return (Summary)Activator.CreateInstance(summaryType, constructArgs);
+			}
+			throw new Exception("Summary type " + summaryType + " not found!");
+		}
+
+		static Dictionary<string, System.Type> summaryNameLookup = null;
+	};
+
+	abstract class Summary
+	{
+
+		public class CaptureRange
+        {
+            public string name;
+            public string startEvent;
+            public string endEvent;
+            public bool includeFirstFrame;
+            public bool includeLastFrame;
+            public CaptureRange(string inName, string start, string end)
+            {
+                name = inName;
+                startEvent = start;
+                endEvent = end;
+                includeFirstFrame = false;
+                includeLastFrame = false;
+            }
+        }
+        public class CaptureData
+        {
+            public int startIndex;
+            public int endIndex;
+            public List<float> Frames;
+            public CaptureData(int start, int end, List<float> inFrames)
+            {
+                startIndex = start;
+                endIndex = end;
+                Frames = inFrames;
+            }
+        }
+
+        public Summary()
+        {
+            stats = new List<string>();
+            captures = new List<CaptureRange>();
+            StatThresholds = new Dictionary<string, ColourThresholdList>();
+
+        }
+        public virtual void WriteSummaryData(System.IO.StreamWriter htmlFile, CsvStats csvStats, bool bWriteSummaryCsv, SummaryTableRowData rowData, string htmlFileName)
+        { }
+
+        public virtual void PostInit(ReportTypeInfo reportTypeInfo, CsvStats csvStats)
+        {
+			// Resolve wildcards and remove duplicates
+			stats = csvStats.GetStatNamesMatchingStringList(stats.ToArray());
+		}
+
+		public abstract string GetName();
+
+        public void ReadStatsFromXML(XElement element)
+        {
+			useUnstrippedCsvStats = element.GetSafeAttibute<bool>("useUnstrippedCsvStats", false);
+			XElement statsElement = element.Element("stats");
+			if (statsElement != null)
+			{
+				stats = statsElement.Value.Split(',').ToList();
+			}
+            foreach (XElement child in element.Elements())
+            {
+                if (child.Name == "capture")
+                {
+                    string captureName = child.Attribute("name").Value;
+                    string captureStart = child.Attribute("startEvent").Value;
+                    string captureEnd = child.Attribute("endEvent").Value;
+                    bool incFirstFrame = Convert.ToBoolean(child.Attribute("includeFirstFrame").Value);
+                    bool incLastFrame = Convert.ToBoolean(child.Attribute("includeLastFrame").Value);
+                    CaptureRange newRange = new CaptureRange(captureName, captureStart, captureEnd);
+                    newRange.includeFirstFrame = incFirstFrame;
+                    newRange.includeLastFrame = incLastFrame;
+                    captures.Add(newRange);
+                }
+                else if (child.Name == "colourThresholds")
+                {
+                    if (child.Attribute("stat") == null)
+                    {
+                        continue;
+                    }
+                    string statName = child.Attribute("stat").Value;
+                    string[] hitchThresholdsStrList = child.Value.Split(',');
+					ColourThresholdList HitchThresholds = new ColourThresholdList();
+					for (int i = 0; i < hitchThresholdsStrList.Length; i++)
+                    {
+						string hitchThresholdStr = hitchThresholdsStrList[i];
+						double thresholdValue = 0.0;
+						string hitchThresholdNumStr = hitchThresholdStr;
+						Colour thresholdColour = null;
+
+						int openBracketIndex = hitchThresholdStr.IndexOf('(');
+						if (openBracketIndex != -1 )
+						{
+							hitchThresholdNumStr = hitchThresholdStr.Substring(0, openBracketIndex);
+							int closeBracketIndex = hitchThresholdStr.IndexOf(')');
+							if (closeBracketIndex > openBracketIndex)
+							{
+								string colourString = hitchThresholdStr.Substring(openBracketIndex+1, closeBracketIndex - openBracketIndex-1);
+								thresholdColour = new Colour(colourString);
+							}
+						}
+						thresholdValue = Convert.ToDouble(hitchThresholdNumStr, System.Globalization.CultureInfo.InvariantCulture);
+
+						HitchThresholds.Add(new ThresholdInfo(thresholdValue, thresholdColour));
+                    }
+                    if (HitchThresholds.Count == 4)
+                    {
+                        StatThresholds.Add(statName, HitchThresholds);
+                    }
+                }
+            }
+        }
+        public CaptureData GetFramesForCapture(CaptureRange inCapture, List<float> FrameTimes, List<CsvEvent> EventsCaptured)
+        {
+            List<float> ReturnFrames = new List<float>();
+            int startFrame = -1;
+            int endFrame = FrameTimes.Count;
+            for (int i = 0; i < EventsCaptured.Count; i++)
+            {
+                if (startFrame < 0 && EventsCaptured[i].Name.ToLower().Contains(inCapture.startEvent.ToLower()))
+                {
+                    startFrame = EventsCaptured[i].Frame;
+                    if (!inCapture.includeFirstFrame)
+                    {
+                        startFrame++;
+                    }
+                }
+                else if (endFrame >= FrameTimes.Count && EventsCaptured[i].Name.ToLower().Contains(inCapture.endEvent.ToLower()))
+                {
+                    endFrame = EventsCaptured[i].Frame;
+                    if (!inCapture.includeLastFrame)
+                    {
+                        endFrame--;
+                    }
+                }
+            }
+            if (startFrame == -1 || endFrame == FrameTimes.Count || endFrame < startFrame)
+            {
+                return null;
+            }
+            ReturnFrames = FrameTimes.GetRange(startFrame, (endFrame - startFrame));
+            CaptureData CaptureToUse = new CaptureData(startFrame, endFrame, ReturnFrames);
+            return CaptureToUse;
+        }
+
+        public string[] GetUniqueStatNames()
+        {
+            HashSet<string> uniqueStats = new HashSet<string>();
+            foreach (string stat in stats)
+            {
+                if (!uniqueStats.Contains(stat))
+                {
+                    uniqueStats.Add(stat);
+                }
+            }
+            return uniqueStats.ToArray();
+        }
+
+		protected double [] ReadColourThresholdsXML(XElement colourThresholdEl)
+		{
+			if (colourThresholdEl != null)
+			{
+				string[] colourStrings = colourThresholdEl.Value.Split(',');
+				if (colourStrings.Length != 4)
+				{
+					throw new Exception("Incorrect number of colourthreshold entries. Should be 4.");
+				}
+				double [] colourThresholds = new double[4];
+				for (int i = 0; i < colourStrings.Length; i++)
+				{
+					colourThresholds[i] = Convert.ToDouble(colourStrings[i], System.Globalization.CultureInfo.InvariantCulture);
+				}
+				return colourThresholds;
+			}
+			return null;
+		}
+
+        public string GetStatThresholdColour(string StatToUse, double value)
+        {
+			ColourThresholdList Thresholds = GetStatColourThresholdList(StatToUse);
+			if (Thresholds != null)
+            {
+				return Thresholds.GetColourForValue(value);
+            }
+			return "'#ffffff'";
+        }
+
+		public ColourThresholdList GetStatColourThresholdList(string StatToUse)
+		{
+			if (StatThresholds.ContainsKey(StatToUse))
+			{
+				return StatThresholds[StatToUse];
+			}
+			return null;
+		}
+
+        public List<CaptureRange> captures;
+        public List<string> stats;
+        public Dictionary<string, ColourThresholdList> StatThresholds;
+		public bool useUnstrippedCsvStats;
+    };
+
+
+
+
+
+
+
+
+
+
+}
diff --git a/Engine/Source/Programs/CSVTools/PerfReportTool/SummaryTable.cs b/Engine/Source/Programs/CSVTools/PerfReportTool/SummaryTable.cs
new file mode 100644
--- /dev/null
+++ b/Engine/Source/Programs/CSVTools/PerfReportTool/SummaryTable.cs
@@ -0,0 +1,1600 @@
+// Copyright Epic Games, Inc. All Rights Reserved.
+
+using System;
+using System.Collections.Generic;
+using System.Linq;
+using System.Xml.Linq;
+using PerfReportTool;
+using System.IO;
+
+namespace PerfSummaries
+{
+	class TableUtil
+	{
+		public static string FormatStatName(string inStatName)
+		{
+			return inStatName.Replace("/", "/ ");
+		}
+
+		public static string SanitizeHtmlString(string str)
+		{
+			return str.Replace("<", "&lt;").Replace(">", "&gt;");
+		}
+
+		public static string SafeTruncateHtmlTableValue(string inValue, int maxLength)
+		{
+			if (inValue.StartsWith("<a") && inValue.EndsWith("</a>"))
+			{
+				// Links require special handling. Only truncate what's inside
+				int openAnchorEndIndex = inValue.IndexOf(">");
+				int closeAnchorStartIndex = inValue.IndexOf("</a>");
+				if (openAnchorEndIndex > 2 && closeAnchorStartIndex > openAnchorEndIndex)
+				{
+					string anchor = inValue.Substring(0, openAnchorEndIndex + 1);
+					string text = inValue.Substring(openAnchorEndIndex + 1, closeAnchorStartIndex - (openAnchorEndIndex + 1));
+					if (text.Length > maxLength)
+					{
+						text = SanitizeHtmlString(text.Substring(0, maxLength)) + "...";
+					}
+					return anchor + text + "</a>";
+				}
+			}
+			return SanitizeHtmlString(inValue.Substring(0, maxLength)) + "...";
+		}
+	}
+
+	class SummaryTableElement
+	{
+		// Bump this when making changes!
+		public static int CacheVersion = 1;
+
+		// NOTE: this is serialized. Don't change the order!
+		public enum Type
+		{
+			CsvStatAverage,
+			CsvMetadata,
+			SummaryTableMetric,
+			ToolMetadata
+		};
+
+		public enum Flags
+		{
+			Hidden = 0x01
+		};
+
+		private SummaryTableElement()
+		{
+
+		}
+		public SummaryTableElement(Type inType, string inName, double inValue, ColourThresholdList inColorThresholdList, string inToolTip, uint inFlags = 0)
+		{
+			type = inType;
+			name = inName;
+			isNumeric = true;
+			numericValue = inValue;
+			value = inValue.ToString();
+			colorThresholdList = inColorThresholdList;
+			tooltip = inToolTip;
+			flags = inFlags;
+		}
+		public SummaryTableElement(Type inType, string inName, string inValue, ColourThresholdList inColorThresholdList, string inToolTip, uint inFlags = 0)
+		{
+			type = inType;
+			name = inName;
+			numericValue = 0.0;
+			isNumeric = false;
+			colorThresholdList = inColorThresholdList;
+			value = inValue;
+			tooltip = inToolTip;
+			flags = inFlags;
+		}
+
+		public static SummaryTableElement ReadFromCache(BinaryReader reader)
+		{
+			SummaryTableElement val = new SummaryTableElement();
+			val.type = (Type)reader.ReadUInt32();
+			val.name = reader.ReadString();
+			val.value = reader.ReadString();
+			val.tooltip = reader.ReadString();
+			val.numericValue = reader.ReadDouble();
+			val.isNumeric = reader.ReadBoolean();
+			val.flags = reader.ReadUInt32();
+			bool hasThresholdList = reader.ReadBoolean();
+			if (hasThresholdList)
+			{
+				int thresholdCount = reader.ReadInt32();
+				val.colorThresholdList = new ColourThresholdList();
+				for (int i = 0; i < thresholdCount; i++)
+				{
+					bool bHasColour = reader.ReadBoolean();
+					Colour thresholdColour = null;
+					if (bHasColour)
+					{
+						thresholdColour = new Colour(reader.ReadString());
+					}
+					double thresholdValue = reader.ReadDouble();
+					ThresholdInfo info = new ThresholdInfo(thresholdValue, thresholdColour);
+					val.colorThresholdList.Add(info);
+				}
+			}
+			return val;
+		}
+
+
+
+		public void WriteToCache(BinaryWriter writer)
+		{
+			writer.Write((uint)type);
+			writer.Write(name);
+			writer.Write(value);
+			writer.Write(tooltip);
+			writer.Write(numericValue);
+			writer.Write(isNumeric);
+			writer.Write(flags);
+			writer.Write(colorThresholdList != null);
+			if (colorThresholdList != null)
+			{
+				writer.Write((int)colorThresholdList.Count);
+				foreach (ThresholdInfo thresholdInfo in colorThresholdList.Thresholds)
+				{
+					writer.Write(thresholdInfo.colour != null);
+					if (thresholdInfo.colour != null)
+					{
+						writer.Write(thresholdInfo.colour.ToString());
+					}
+					writer.Write(thresholdInfo.value);
+				}
+			}
+		}
+
+		public SummaryTableElement Clone()
+		{
+			return (SummaryTableElement)MemberwiseClone();
+		}
+
+		public void SetFlag(Flags flag, bool value)
+		{
+			if (value)
+			{
+				flags |= (uint)flag;
+			}
+			else
+			{
+				flags &= ~(uint)flag;
+			}
+		}
+		public bool GetFlag(Flags flag)
+		{
+			return (flags & (uint)flag) != 0;
+		}
+
+		public Type type;
+		public string name;
+		public string value;
+		public string tooltip;
+		public ColourThresholdList colorThresholdList;
+		public double numericValue;
+		public bool isNumeric;
+		public uint flags;
+	}
+	class SummaryTableRowData
+	{
+		public SummaryTableRowData()
+		{
+		}
+
+		// TODO: If this is bumped beyond 6, we need to implement backwards compatibility
+		static int CacheVersion = 6;
+
+		public static SummaryTableRowData TryReadFromCache(string summaryTableCacheDir, string csvId)
+		{
+			string filename = Path.Combine(summaryTableCacheDir, csvId + ".prc");
+			return TryReadFromCacheFile(filename);
+		}
+
+		public static SummaryTableRowData TryReadFromCacheFile(string filename, bool bReadJustInitialMetadata = false)
+		{
+			SummaryTableRowData metaData = null;
+			if (!File.Exists(filename))
+			{
+				return null;
+			}
+			try
+			{
+				using (FileStream fileStream = new FileStream(filename, FileMode.Open, FileAccess.Read))
+				{
+					BinaryReader reader = new BinaryReader(fileStream);
+					int version = reader.ReadInt32();
+					int metadataValueVersion = reader.ReadInt32();
+					if (version == CacheVersion && metadataValueVersion == SummaryTableElement.CacheVersion)
+					{
+						bool bEarlyOut = false;
+						metaData = new SummaryTableRowData();
+						int dictEntryCount = reader.ReadInt32();
+						for (int i = 0; i < dictEntryCount; i++)
+						{
+							string key = reader.ReadString();
+							SummaryTableElement value = SummaryTableElement.ReadFromCache(reader);
+							// If we're just reading initial metadata then skip everything after ToolMetadata and CsvMetadata
+							if (bReadJustInitialMetadata && value.type != SummaryTableElement.Type.ToolMetadata && value.type != SummaryTableElement.Type.CsvMetadata)
+							{
+								bEarlyOut = true;
+								break;
+							}
+							metaData.dict.Add(key, value);
+						}
+
+						if (!bEarlyOut)
+						{
+							string endString = reader.ReadString();
+							if (endString != "END")
+							{
+								Console.WriteLine("Corruption detected in " + filename + ". Skipping read");
+								metaData = null;
+							}
+						}
+					}
+					reader.Close();
+				}
+			}
+			catch (Exception e)
+			{
+				metaData = null;
+				Console.WriteLine("Error reading from cache file " + filename + ": " + e.Message);
+			}
+			return metaData;
+		}
+
+		public bool WriteToCache(string summaryTableCacheDir, string csvId)
+		{
+			string filename = Path.Combine(summaryTableCacheDir, csvId + ".prc");
+			try
+			{
+				using (FileStream fileStream = new FileStream(filename, FileMode.Create))
+				{
+					BinaryWriter writer = new BinaryWriter(fileStream);
+					writer.Write(CacheVersion);
+					writer.Write(SummaryTableElement.CacheVersion);
+
+					writer.Write(dict.Count);
+					foreach (KeyValuePair<string, SummaryTableElement> entry in dict)
+					{
+						writer.Write(entry.Key);
+						entry.Value.WriteToCache(writer);
+					}
+					writer.Write("END");
+					writer.Close();
+				}
+			}
+			catch (IOException)
+			{
+				Console.WriteLine("Failed to write to cache file " + filename + ".");
+				return false;
+			}
+			return true;
+		}
+
+		public void RemoveSafe(string name)
+		{
+			string key = name.ToLower();
+			if (dict.ContainsKey(key))
+			{
+				dict.Remove(key);
+			}
+		}
+
+		public void Add(SummaryTableElement.Type type, string name, double value, ColourThresholdList colorThresholdList = null, string tooltip = "", uint flags = 0)
+		{
+			string key = name.ToLower();
+			SummaryTableElement metadataValue = new SummaryTableElement(type, name, value, colorThresholdList, tooltip, flags);
+			try
+			{
+				dict.Add(key, metadataValue);
+			}
+			catch (System.ArgumentException)
+			{
+				throw new Exception("Summary metadata key " + key + " has already been added");
+			}
+		}
+
+		public void Add(SummaryTableElement.Type type, string name, string value, ColourThresholdList colorThresholdList = null, string tooltip = "", uint flags = 0)
+		{
+			string key = name.ToLower();
+			double numericValue = double.MaxValue;
+			try
+			{
+				numericValue = Convert.ToDouble(value, System.Globalization.CultureInfo.InvariantCulture);
+			}
+			catch { }
+
+			SummaryTableElement metadataValue = null;
+			if (numericValue != double.MaxValue)
+			{
+				metadataValue = new SummaryTableElement(type, name, numericValue, colorThresholdList, tooltip, flags);
+			}
+			else
+			{
+				metadataValue = new SummaryTableElement(type, name, value, colorThresholdList, tooltip, flags);
+			}
+
+			try
+			{
+				dict.Add(key, metadataValue);
+			}
+			catch (System.ArgumentException)
+			{
+				throw new Exception("Summary metadata key " + key + " has already been added");
+			}
+		}
+
+		public void AddString(SummaryTableElement.Type type, string name, string value, ColourThresholdList colorThresholdList = null, string tooltip = "")
+		{
+			string key = name.ToLower();
+			SummaryTableElement metadataValue = new SummaryTableElement(type, name, value, colorThresholdList, tooltip);
+			dict.Add(key, metadataValue);
+		}
+
+		public Dictionary<string, SummaryTableElement> dict = new Dictionary<string, SummaryTableElement>();
+	};
+
+	class SummarySectionBoundaryInfo
+	{
+		public SummarySectionBoundaryInfo(string inStatName, string inStartToken, string inEndToken, int inLevel, bool inInCollatedTable, bool inInFullTable)
+		{
+			statName = inStatName;
+			startToken = inStartToken;
+			endToken = inEndToken;
+			level = inLevel;
+			inCollatedTable = inInCollatedTable;
+			inFullTable = inInFullTable;
+		}
+		public string statName;
+		public string startToken;
+		public string endToken;
+		public int level;
+		public bool inCollatedTable;
+		public bool inFullTable;
+	};
+
+	class SummaryTableInfo
+	{
+		public SummaryTableInfo(XElement tableElement)
+		{
+			XAttribute rowSortAt = tableElement.Attribute("rowSort");
+			if (rowSortAt != null)
+			{
+				rowSortList.AddRange(rowSortAt.Value.Split(','));
+			}
+			XAttribute weightByColumnAt = tableElement.Attribute("weightByColumn");
+			if (weightByColumnAt != null)
+			{
+				weightByColumn = weightByColumnAt.Value.ToLower();
+			}
+
+			XElement filterEl = tableElement.Element("filter");
+			if (filterEl != null)
+			{
+				columnFilterList.AddRange(filterEl.Value.Split(','));
+			}
+
+			foreach (XElement sectionBoundaryEl in tableElement.Elements("sectionBoundary"))
+			{
+				if (sectionBoundaryEl != null)
+				{
+					SummarySectionBoundaryInfo sectionBoundary = new SummarySectionBoundaryInfo(
+						sectionBoundaryEl.GetSafeAttibute<string>("statName"),
+						sectionBoundaryEl.GetSafeAttibute<string>("startToken"),
+						sectionBoundaryEl.GetSafeAttibute<string>("endToken"),
+						sectionBoundaryEl.GetSafeAttibute<int>("level", 0),
+						sectionBoundaryEl.GetSafeAttibute<bool>("inCollatedTable", true),
+						sectionBoundaryEl.GetSafeAttibute<bool>("inFullTable", true)
+						);
+					sectionBoundaries.Add(sectionBoundary);
+				}
+			}
+		}
+
+		public SummaryTableInfo(string filterListStr, string rowSortStr)
+		{
+			columnFilterList.AddRange(filterListStr.Split(','));
+			rowSortList.AddRange(rowSortStr.Split(','));
+		}
+
+		public List<string> rowSortList = new List<string>();
+		public List<string> columnFilterList = new List<string>();
+		public List<SummarySectionBoundaryInfo> sectionBoundaries = new List<SummarySectionBoundaryInfo>();
+		public string weightByColumn = null;
+	}
+
+
+	class SummaryTableColumnFormatInfoCollection
+	{
+		public SummaryTableColumnFormatInfoCollection(XElement element)
+		{
+			foreach (XElement child in element.Elements("columnInfo"))
+			{
+				columnFormatInfoList.Add(new SummaryTableColumnFormatInfo(child));
+			}
+		}
+
+		public SummaryTableColumnFormatInfo GetFormatInfo(string columnName)
+		{
+			string lowerColumnName = columnName.ToLower();
+			if (lowerColumnName.StartsWith("avg ") || lowerColumnName.StartsWith("min ") || lowerColumnName.StartsWith("max "))
+			{
+				lowerColumnName = lowerColumnName.Substring(4);
+			}
+			foreach (SummaryTableColumnFormatInfo columnInfo in columnFormatInfoList)
+			{
+				int wildcardIndex = columnInfo.name.IndexOf('*');
+				if (wildcardIndex == -1)
+				{
+					if (columnInfo.name == lowerColumnName)
+					{
+						return columnInfo;
+					}
+				}
+				else
+				{
+					string prefix = columnInfo.name.Substring(0, wildcardIndex);
+					if (lowerColumnName.StartsWith(prefix))
+					{
+						return columnInfo;
+					}
+				}
+			}
+			return defaultColumnInfo;
+		}
+
+		public static SummaryTableColumnFormatInfo DefaultColumnInfo
+		{
+			get { return defaultColumnInfo; }
+		}
+
+		List<SummaryTableColumnFormatInfo> columnFormatInfoList = new List<SummaryTableColumnFormatInfo>();
+		static SummaryTableColumnFormatInfo defaultColumnInfo = new SummaryTableColumnFormatInfo();
+	};
+
+	enum AutoColorizeMode
+	{
+		Off,
+		HighIsBad,
+		LowIsBad,
+	};
+
+	class SummaryTableColumnFormatInfo
+	{
+		public SummaryTableColumnFormatInfo()
+		{
+			name = "Default";
+		}
+		public SummaryTableColumnFormatInfo(XElement element)
+		{
+			name = element.Attribute("name").Value.ToLower();
+
+			string autoColorizeStr = element.GetSafeAttibute<string>("autoColorize", "highIsBad").ToLower();
+			var modeList = Enum.GetValues(typeof(AutoColorizeMode));
+			foreach (AutoColorizeMode mode in modeList)
+			{
+				if (mode.ToString().ToLower() == autoColorizeStr)
+				{
+					autoColorizeMode = mode;
+					break;
+				}
+			}
+			numericFormat = element.GetSafeAttibute<string>("numericFormat");
+		}
+		public AutoColorizeMode autoColorizeMode = AutoColorizeMode.HighIsBad;
+		public string name;
+		public string numericFormat;
+	};
+
+	class SummaryTableColumn
+	{
+		public string name;
+		public bool isNumeric = false;
+		public string displayName;
+		public bool isRowWeightColumn = false;
+		List<double> doubleValues = new List<double>();
+		List<string> stringValues = new List<string>();
+		List<string> toolTips = new List<string>();
+
+		List<ColourThresholdList> colourThresholds = new List<ColourThresholdList>();
+		ColourThresholdList colourThresholdOverride = null;
+		public SummaryTableColumn(string inName, bool inIsNumeric, string inDisplayName = null, bool inIsRowWeightColumn = false)
+		{
+			name = inName;
+			isNumeric = inIsNumeric;
+			displayName = inDisplayName;
+			isRowWeightColumn = inIsRowWeightColumn;
+		}
+		public SummaryTableColumn Clone()
+		{
+			SummaryTableColumn newColumn = new SummaryTableColumn(name, isNumeric, displayName, isRowWeightColumn);
+			newColumn.doubleValues.AddRange(doubleValues);
+			newColumn.stringValues.AddRange(stringValues);
+			newColumn.colourThresholds.AddRange(colourThresholds);
+			newColumn.toolTips.AddRange(toolTips);
+			return newColumn;
+		}
+
+		public string GetDisplayName()
+		{
+			if (displayName == null)
+			{
+				return TableUtil.FormatStatName(name);
+			}
+			return displayName;
+		}
+
+		public void SetValue(int index, double value)
+		{
+			if (!isNumeric)
+			{
+				// This is already a non-numeric column. Better treat this as a string value
+				SetStringValue(index, value.ToString());
+				return;
+			}
+			// Grow to fill if necessary
+			if (index >= doubleValues.Count)
+			{
+				for (int i = doubleValues.Count; i <= index; i++)
+				{
+					doubleValues.Add(double.MaxValue);
+				}
+			}
+			doubleValues[index] = value;
+		}
+
+		void convertToStrings()
+		{
+			if (isNumeric)
+			{
+				stringValues = new List<string>();
+				foreach (float f in doubleValues)
+				{
+					stringValues.Add(f.ToString());
+				}
+				doubleValues = new List<double>();
+				isNumeric = false;
+			}
+		}
+
+		public void SetColourThresholds(int index, ColourThresholdList value)
+		{
+			// Grow to fill if necessary
+			if (index >= colourThresholds.Count)
+			{
+				for (int i = colourThresholds.Count; i <= index; i++)
+				{
+					colourThresholds.Add(null);
+				}
+			}
+			colourThresholds[index] = value;
+		}
+
+		public ColourThresholdList GetColourThresholds(int index)
+		{
+			if (index < colourThresholds.Count)
+			{
+				return colourThresholds[index];
+			}
+			return null;
+		}
+
+		public string GetColour(int index)
+		{
+			ColourThresholdList thresholds = null;
+			double value = GetValue(index);
+			if (value == double.MaxValue)
+			{
+				return null;
+			}
+			if (colourThresholdOverride != null)
+			{
+				thresholds = colourThresholdOverride;
+			}
+			else
+			{
+				if (index < colourThresholds.Count)
+				{
+					thresholds = colourThresholds[index];
+				}
+				if (thresholds == null)
+				{
+					return null;
+				}
+			}
+			return thresholds.GetColourForValue(value);
+		}
+
+		public void ComputeAutomaticColourThresholds(AutoColorizeMode autoColorizeMode)
+		{
+			if (autoColorizeMode == AutoColorizeMode.Off)
+			{
+				return;
+			}
+			colourThresholds = new List<ColourThresholdList>();
+			double maxValue = -double.MaxValue;
+			double minValue = double.MaxValue;
+			double totalValue = 0.0f;
+			double validCount = 0.0f;
+			for (int i = 0; i < doubleValues.Count; i++)
+			{
+				double val = doubleValues[i];
+				if (val != double.MaxValue)
+				{
+					maxValue = Math.Max(val, maxValue);
+					minValue = Math.Min(val, minValue);
+					totalValue += val;
+					validCount += 1.0f;
+				}
+			}
+			if (minValue == maxValue)
+			{
+				return;
+			}
+
+			Colour green = new Colour(0.4f, 0.82f, 0.45f);
+			Colour yellow = new Colour(1.0f, 1.0f, 0.5f);
+			Colour red = new Colour(1.0f, 0.4f, 0.4f);
+
+			double averageValue = totalValue / validCount; // TODO: Weighted average 
+			colourThresholdOverride = new ColourThresholdList();
+			colourThresholdOverride.Add(new ThresholdInfo(minValue, (autoColorizeMode == AutoColorizeMode.HighIsBad) ? green : red));
+			colourThresholdOverride.Add(new ThresholdInfo(averageValue, yellow));
+			colourThresholdOverride.Add(new ThresholdInfo(averageValue, yellow));
+			colourThresholdOverride.Add(new ThresholdInfo(maxValue, (autoColorizeMode == AutoColorizeMode.HighIsBad) ? red : green));
+		}
+
+		public int GetCount()
+		{
+			return Math.Max(doubleValues.Count, stringValues.Count);
+		}
+		public double GetValue(int index)
+		{
+			if (index >= doubleValues.Count)
+			{
+				return double.MaxValue;
+			}
+			return doubleValues[index];
+		}
+
+		public void SetStringValue(int index, string value)
+		{
+			if (isNumeric)
+			{
+				// Better convert this to a string column, since we're trying to add a string to it
+				convertToStrings();
+			}
+			// Grow to fill if necessary
+			if (index >= stringValues.Count)
+			{
+				for (int i = stringValues.Count; i <= index; i++)
+				{
+					stringValues.Add("");
+				}
+			}
+			stringValues[index] = value;
+			isNumeric = false;
+		}
+		public string GetStringValue(int index, bool roundNumericValues = false, string forceNumericFormat = null)
+		{
+			if (isNumeric)
+			{
+				if (index >= doubleValues.Count || doubleValues[index] == double.MaxValue)
+				{
+					return "";
+				}
+				double val = doubleValues[index];
+				if (forceNumericFormat != null)
+				{
+					return val.ToString(forceNumericFormat);
+				}
+				else if (roundNumericValues)
+				{
+					double absVal = Math.Abs(val);
+					double frac = absVal - (double)Math.Truncate(absVal);
+					if (absVal >= 250.0f || frac < 0.0001f)
+					{
+						return val.ToString("0");
+					}
+					if (absVal >= 50.0f)
+					{
+						return val.ToString("0.0");
+					}
+					if (absVal >= 0.1)
+					{
+						return val.ToString("0.00");
+					}
+					return val.ToString("0.000");
+				}
+				return val.ToString();
+			}
+			else
+			{
+				if (index >= stringValues.Count)
+				{
+					return "";
+				}
+				if (forceNumericFormat != null)
+				{
+					// We're forcing a numeric format on something that's technically a string, but since we were asked, we'll try to do it anyway 
+					// Note: this is not ideal, but it's useful for collated table columns, which might get converted to non-numeric during collation
+					try
+					{
+						return Convert.ToDouble(stringValues[index], System.Globalization.CultureInfo.InvariantCulture).ToString(forceNumericFormat);
+					}
+					catch { } // Ignore. Just fall through...
+				}
+				return stringValues[index];
+			}
+		}
+		public void SetToolTipValue(int index, string value)
+		{
+			// Grow to fill if necessary
+			if (index >= toolTips.Count)
+			{
+				for (int i = toolTips.Count; i <= index; i++)
+				{
+					toolTips.Add("");
+				}
+			}
+			toolTips[index] = value;
+		}
+		public string GetToolTipValue(int index)
+		{
+			if (index >= toolTips.Count)
+			{
+				return "";
+			}
+			return toolTips[index];
+		}
+
+	};
+
+
+
+	class SummaryTable
+	{
+		public SummaryTable()
+		{
+		}
+
+		public SummaryTable CollateSortedTable(List<string> collateByList, bool addMinMaxColumns)
+		{
+			int numSubColumns = addMinMaxColumns ? 3 : 1;
+
+			// Find all the columns in collateByList
+			HashSet<SummaryTableColumn> collateByColumns = new HashSet<SummaryTableColumn>();
+			foreach (string collateBy in collateByList)
+			{
+				string key = collateBy.ToLower();
+				if (columnLookup.ContainsKey(key))
+				{
+					collateByColumns.Add(columnLookup[key]);
+				}
+			}
+			if (collateByColumns.Count == 0)
+			{
+				throw new Exception("None of the metadata strings were found:" + collateByList.ToString());
+			}
+
+			// Add the new collateBy columns in the order they appear in the original column list
+			List<SummaryTableColumn> newColumns = new List<SummaryTableColumn>();
+			List<string> finalSortByList = new List<string>();
+			foreach (SummaryTableColumn srcColumn in columns)
+			{
+				if (collateByColumns.Contains(srcColumn))
+				{
+					newColumns.Add(new SummaryTableColumn(srcColumn.name, false, srcColumn.displayName));
+					finalSortByList.Add(srcColumn.name.ToLower());
+					// Early out if we've found all the columns
+					if (finalSortByList.Count == collateByColumns.Count)
+					{
+						break;
+					}
+				}
+			}
+
+			newColumns.Add(new SummaryTableColumn("Count", true));
+			int countColumnIndex = newColumns.Count - 1;
+
+			int numericColumnStartIndex = newColumns.Count;
+			List<int> srcToDestBaseColumnIndex = new List<int>();
+			foreach (SummaryTableColumn column in columns)
+			{
+				// Add avg/min/max columns for this column if it's numeric and we didn't already add it above 
+				if (column.isNumeric && !collateByColumns.Contains(column))
+				{
+					srcToDestBaseColumnIndex.Add(newColumns.Count);
+					newColumns.Add(new SummaryTableColumn("Avg " + column.name, true));
+					if (addMinMaxColumns)
+					{
+						newColumns.Add(new SummaryTableColumn("Min " + column.name, true));
+						newColumns.Add(new SummaryTableColumn("Max " + column.name, true));
+					}
+				}
+				else
+				{
+					srcToDestBaseColumnIndex.Add(-1);
+				}
+			}
+
+			List<double> RowMaxValues = new List<double>();
+			List<double> RowTotals = new List<double>();
+			List<double> RowMinValues = new List<double>();
+			List<int> RowCounts = new List<int>();
+			List<double> RowWeights = new List<double>();
+			List<ColourThresholdList> RowColourThresholds = new List<ColourThresholdList>();
+
+			// Set the initial sort key
+			string CurrentRowSortKey = "";
+			foreach (string collateBy in finalSortByList)
+			{
+				CurrentRowSortKey += "{" + columnLookup[collateBy].GetStringValue(0) + "}";
+			}
+
+			int destRowIndex = 0;
+			bool reset = true;
+			int mergedRowsCount = 0;
+			for (int i = 0; i < rowCount; i++)
+			{
+				if (reset)
+				{
+					RowMaxValues.Clear();
+					RowMinValues.Clear();
+					RowTotals.Clear();
+					RowCounts.Clear();
+					RowWeights.Clear();
+					RowColourThresholds.Clear();
+					for (int j = 0; j < columns.Count; j++)
+					{
+						if (addMinMaxColumns)
+						{
+							RowMaxValues.Add(-double.MaxValue);
+							RowMinValues.Add(double.MaxValue);
+						}
+						RowTotals.Add(0.0f);
+						RowCounts.Add(0);
+						RowWeights.Add(0.0);
+						RowColourThresholds.Add(null);
+					}
+					mergedRowsCount = 0;
+					reset = false;
+				}
+
+				// Compute min/max/total for all numeric columns
+				for (int j = 0; j < columns.Count; j++)
+				{
+					SummaryTableColumn column = columns[j];
+					if (column.isNumeric)
+					{
+						double value = column.GetValue(i);
+						if (value != double.MaxValue)
+						{
+							if (addMinMaxColumns)
+							{
+								RowMaxValues[j] = Math.Max(RowMaxValues[j], value);
+								RowMinValues[j] = Math.Min(RowMinValues[j], value);
+							}
+							RowColourThresholds[j] = column.GetColourThresholds(i);
+							RowCounts[j]++;
+							double rowWeight = (rowWeightings != null) ? rowWeightings[i] : 1.0;
+							RowWeights[j] += rowWeight;
+							RowTotals[j] += value * rowWeight;
+						}
+					}
+				}
+				mergedRowsCount++;
+
+				// Are we done?
+				string nextSortKey = "";
+				if (i < rowCount - 1)
+				{
+					foreach (string collateBy in finalSortByList)
+					{
+						nextSortKey += "{" + columnLookup[collateBy].GetStringValue(i + 1) + "}";
+					}
+				}
+
+				// If this is the last row or if the sort key is different then write it out
+				if (nextSortKey != CurrentRowSortKey)
+				{
+					for (int j = 0; j < countColumnIndex; j++)
+					{
+						string key = newColumns[j].name.ToLower();
+						newColumns[j].SetStringValue(destRowIndex, columnLookup[key].GetStringValue(i));
+					}
+					// Commit the row 
+					newColumns[countColumnIndex].SetValue(destRowIndex, (double)mergedRowsCount);
+					for (int j = 0; j < columns.Count; j++)
+					{
+						int destColumnBaseIndex = srcToDestBaseColumnIndex[j];
+						if (destColumnBaseIndex != -1 && RowCounts[j] > 0)
+						{
+							newColumns[destColumnBaseIndex].SetValue(destRowIndex, RowTotals[j] / RowWeights[j]);
+							if (addMinMaxColumns)
+							{
+								newColumns[destColumnBaseIndex + 1].SetValue(destRowIndex, RowMinValues[j]);
+								newColumns[destColumnBaseIndex + 2].SetValue(destRowIndex, RowMaxValues[j]);
+							}
+
+							// Set colour thresholds based on the source column
+							ColourThresholdList Thresholds = RowColourThresholds[j];
+							for (int k = 0; k < numSubColumns; k++)
+							{
+								newColumns[destColumnBaseIndex + k].SetColourThresholds(destRowIndex, Thresholds);
+							}
+						}
+					}
+					reset = true;
+					destRowIndex++;
+				}
+				CurrentRowSortKey = nextSortKey;
+			}
+
+			SummaryTable newTable = new SummaryTable();
+			newTable.columns = newColumns;
+			newTable.InitColumnLookup();
+			newTable.rowCount = destRowIndex;
+			newTable.firstStatColumnIndex = numericColumnStartIndex;
+			newTable.isCollated = true;
+			newTable.hasMinMaxColumns = addMinMaxColumns;
+			return newTable;
+		}
+
+		public SummaryTable SortAndFilter(string customFilter, string customRowSort = "buildversion,deviceprofile", bool bReverseSort = false, string weightByColumnName = null)
+		{
+			return SortAndFilter(customFilter.Split(',').ToList(), customRowSort.Split(',').ToList(), bReverseSort, weightByColumnName);
+		}
+
+		public SummaryTable SortAndFilter(List<string> columnFilterList, List<string> rowSortList, bool bReverseSort, string weightByColumnName)
+		{
+			SummaryTable newTable = SortRows(rowSortList, bReverseSort);
+
+			// Make a list of all unique keys
+			List<string> allMetadataKeys = new List<string>();
+			Dictionary<string, SummaryTableColumn> nameLookup = new Dictionary<string, SummaryTableColumn>();
+			foreach (SummaryTableColumn col in newTable.columns)
+			{
+				string key = col.name.ToLower();
+				if (!nameLookup.ContainsKey(key))
+				{
+					nameLookup.Add(key, col);
+					allMetadataKeys.Add(key);
+				}
+			}
+			allMetadataKeys.Sort();
+
+			// Generate the list of requested metadata keys that this table includes
+			List<string> orderedKeysWithDupes = new List<string>();
+
+			// Add metadata keys from the column filter list in the order they appear
+			foreach (string filterStr in columnFilterList)
+			{
+				string filterStrLower = filterStr.Trim().ToLower();
+				// Find all matching
+				if (filterStrLower.EndsWith("*"))
+				{
+					string prefix = filterStrLower.Trim('*');
+					// Linear search through the sorted key list
+					bool bFound = false;
+					for (int wildcardSearchIndex = 0; wildcardSearchIndex < allMetadataKeys.Count; wildcardSearchIndex++)
+					{
+						if (allMetadataKeys[wildcardSearchIndex].StartsWith(prefix))
+						{
+							orderedKeysWithDupes.Add(allMetadataKeys[wildcardSearchIndex]);
+							bFound = true;
+						}
+						else if (bFound)
+						{
+							// Early exit: already found one key. If the pattern no longer matches then we must be done
+							break;
+						}
+					}
+				}
+				else
+				{
+					string key = filterStrLower;
+					orderedKeysWithDupes.Add(key);
+				}
+			}
+
+			// Compute row weights
+			if (weightByColumnName != null && nameLookup.ContainsKey(weightByColumnName))
+			{
+				SummaryTableColumn rowWeightColumn = nameLookup[weightByColumnName];
+				newTable.rowWeightings = new List<double>(rowWeightColumn.GetCount());
+				for (int i = 0; i < rowWeightColumn.GetCount(); i++)
+				{
+					newTable.rowWeightings.Add(rowWeightColumn.GetValue(i));
+				}
+			}
+
+			List<SummaryTableColumn> newColumnList = new List<SummaryTableColumn>();
+			// Add all the ordered keys that exist, ignoring duplicates
+			foreach (string key in orderedKeysWithDupes)
+			{
+				if (nameLookup.ContainsKey(key))
+				{
+					newColumnList.Add(nameLookup[key]);
+					// Remove from the list so it doesn't get counted again
+					nameLookup.Remove(key);
+				}
+			}
+
+
+			newTable.columns = newColumnList;
+			newTable.rowCount = rowCount;
+			newTable.InitColumnLookup();
+
+			return newTable;
+		}
+
+		public void ApplyDisplayNameMapping(Dictionary<string, string> statDisplaynameMapping)
+		{
+			// Convert to a display-friendly name
+			foreach (SummaryTableColumn column in columns)
+			{
+				if (statDisplaynameMapping != null && column.displayName == null)
+				{
+					string name = column.name;
+					string suffix = "";
+					string prefix = "";
+					string statName = GetStatNameWithPrefixAndSuffix(name, out prefix, out suffix);
+					if (statDisplaynameMapping.ContainsKey(statName.ToLower()))
+					{
+						column.displayName = prefix + statDisplaynameMapping[statName.ToLower()] + suffix;
+					}
+				}
+			}
+		}
+
+		string GetStatNameWithoutPrefixAndSuffix(string inName)
+		{
+			string suffix = "";
+			string prefix = "";
+			return GetStatNameWithPrefixAndSuffix(inName, out prefix, out suffix);
+		}
+
+		string GetStatNameWithPrefixAndSuffix(string inName, out string prefix, out string suffix)
+		{
+			suffix = "";
+			prefix = "";
+			string statName = inName;
+			if (inName.StartsWith("Avg ") || inName.StartsWith("Max ") || inName.StartsWith("Min "))
+			{
+				prefix = inName.Substring(0, 4);
+				statName = inName.Substring(4);
+			}
+			if (statName.EndsWith(" Avg") || statName.EndsWith(" Max") || statName.EndsWith(" Min"))
+			{
+				suffix = statName.Substring(statName.Length - 4);
+				statName = statName.Substring(0, statName.Length - 4);
+			}
+			return statName;
+		}
+
+		public void WriteToCSV(string csvFilename)
+		{
+			System.IO.StreamWriter csvFile = new System.IO.StreamWriter(csvFilename, false);
+			List<string> headerRow = new List<string>();
+			foreach (SummaryTableColumn column in columns)
+			{
+				headerRow.Add(column.name);
+			}
+			csvFile.WriteLine(string.Join(",", headerRow));
+
+			for (int i = 0; i < rowCount; i++)
+			{
+				List<string> rowStrings = new List<string>();
+				foreach (SummaryTableColumn column in columns)
+				{
+					string cell = column.GetStringValue(i, false);
+					// Sanitize so it opens in a spreadsheet (e.g. for buildversion) 
+					cell = cell.TrimStart('+');
+					rowStrings.Add(cell);
+				}
+				csvFile.WriteLine(string.Join(",", rowStrings));
+			}
+			csvFile.Close();
+		}
+
+
+		public void WriteToHTML(string htmlFilename, string VersionString, bool bSpreadsheetFriendlyStrings, List<SummarySectionBoundaryInfo> sectionBoundaries, bool bScrollableTable, bool bAddMinMaxColumns, int maxColumnStringLength, SummaryTableColumnFormatInfoCollection columnFormatInfoList, string weightByColumnName)
+		{
+			System.IO.StreamWriter htmlFile = new System.IO.StreamWriter(htmlFilename, false);
+			int statColSpan = hasMinMaxColumns ? 3 : 1;
+			int cellPadding = 2;
+			if (isCollated)
+			{
+				cellPadding = 4;
+			}
+
+			htmlFile.WriteLine("<html>");
+			htmlFile.WriteLine("<head><title>Performance summary</title>");
+
+			// Figure out the sticky column count
+			int stickyColumnCount = 0;
+			if (bScrollableTable)
+			{
+				stickyColumnCount = 1;
+				if (isCollated)
+				{
+					for (int i = 0; i < columns.Count; i++)
+					{
+						if (columns[i].name == "Count")
+						{
+							stickyColumnCount = i + 1;
+							break;
+						}
+					}
+				}
+			}
+
+			if (bScrollableTable)
+			{
+				// Insert some javascript to make the columns sticky. It's not possible to do this for multiple columns with pure CSS, since you need to compute the X offset dynamically
+				// We need to do this when the page is loaded or the window is resized
+				htmlFile.WriteLine("<script>");
+
+				htmlFile.WriteLine(
+					"var originalStyleElement = null; \n" +
+					"document.addEventListener('DOMContentLoaded', function(event) { regenerateStickyColumnCss(); }) \n" +
+					"window.addEventListener('resize', function(event) { regenerateStickyColumnCss(); }) \n" +
+					"\n" +
+					"function regenerateStickyColumnCss() { \n" +
+					"  var styleElement=document.getElementById('pageStyle'); \n" +
+					"  var table=document.getElementById('mainTable'); \n" +
+					"  if ( table.rows.length < 2 ) \n" +
+					"	return; \n" +
+					"  if (originalStyleElement == null) \n" +
+					"    originalStyleElement = styleElement.textContent; \n" +
+					"  else \n" +
+					"    styleElement.textContent = originalStyleElement  \n"
+				);
+
+				// Make the columns Sticky and compute their X offsets
+				htmlFile.WriteLine(
+					"  var numStickyCols=" + stickyColumnCount + "; \n" +
+					"  var xOffset=0; \n" +
+					"  for (var i=0;i<numStickyCols;i++) \n" +
+					"  { \n" +
+					"	 var rBorderParam=(i==numStickyCols-1) ? 'border-right: 2px solid black;':''; \n" +
+					"	 styleElement.textContent+='tr.lastHeaderRow th:nth-child('+(i+1)+') {  z-index: 8;  border-top: 2px solid black;  font-size: 11px;  left: '+xOffset+'px;'+rBorderParam+'}'; \n" +
+					"	 styleElement.textContent+='td:nth-child('+(i+1)+') {  position: -webkit-sticky;  position: sticky; z-index: 7;  left: '+xOffset+'px; '+rBorderParam+'}'; \n" +
+					"	 xOffset+=table.rows[1].cells[i].offsetWidth; \n" +
+					"  } \n"
+					);
+				// Make all sticky columns except the count column opaque
+				htmlFile.WriteLine(
+					"  for (var i=0;i<numStickyCols-1;i++) \n" +
+					"  { \n" +
+					"    styleElement.textContent+='tr:nth-child(odd) td:nth-child('+(i+1)+') { background-color: #e2e2e2; }'; \n" +
+					"    styleElement.textContent+='tr:nth-child(even) td:nth-child('+(i+1)+') { background-color: #ffffff; }'; \n" +
+					"  } \n" +
+					"} \n"
+				);
+				htmlFile.WriteLine("</script>");
+			}
+
+			htmlFile.WriteLine("<style type='text/css' id='pageStyle'>");
+			htmlFile.WriteLine("p {  font-family: 'Verdana', Times, serif; font-size: 12px }");
+			htmlFile.WriteLine("h3 {  font-family: 'Verdana', Times, serif; font-size: 14px }");
+			htmlFile.WriteLine("h2 {  font-family: 'Verdana', Times, serif; font-size: 16px }");
+			htmlFile.WriteLine("h1 {  font-family: 'Verdana', Times, serif; font-size: 20px }");
+			string tableCss = "";
+			if (bScrollableTable)
+			{
+				int firstColMaxStringLength = 0;
+				if (columns.Count > 0)
+				{
+					for (int i = 0; i < columns[0].GetCount(); i++)
+					{
+						firstColMaxStringLength = Math.Max(firstColMaxStringLength, columns[0].GetStringValue(i).Length);
+					}
+				}
+				int firstColWidth = (int)(firstColMaxStringLength * 6.5);
+				tableCss =
+					"table {table-layout: fixed;} \n" +
+					"table, th, td { border: 0px solid black; border-spacing: 0; border-collapse: separate; padding: " + cellPadding + "px; vertical-align: center; font-family: 'Verdana', Times, serif; font-size: 10px;} \n" +
+					"td {" +
+					"  border-right: 1px solid black;" +
+					"  max-width: 400;" +
+					"} \n" +
+					"tr:first-element { border-top: 2px; border-bottom: 2px } \n" +
+					"th {" +
+					"  width: 75px;" +
+					"  max-width: 400;" +
+					"  position: -webkit-sticky;" +
+					"  position: sticky;" +
+					"  border-right: 1px solid black;" +
+					"  border-top: 2px solid black;" +
+					"  z-index: 5;" +
+					"  background-color: #ffffff;" +
+					"  top:0;" +
+					"  font-size: 9px;" +
+					"  word-wrap: break-word;" +
+					"  overflow: hidden;" +
+					"  height: 60;" +
+					"} \n";
+
+				// Top-left cell of the table is always on top, big font, thick border
+				tableCss += "tr:first-child th:first-child { z-index: 100;  border-right: 2px solid black; border-top: 2px solid black; font-size: 11px; top:0; left: 0px; } \n";
+
+				// Fix the first column width
+				tableCss += "th:first-child, td:first-child { border-left: 2px solid black; min-width: " + firstColWidth + ";} \n";
+
+				if (bAddMinMaxColumns && isCollated)
+				{
+					tableCss += "tr.lastHeaderRow th { top:60px; height:20px; } \n";
+				}
+
+				if (!isCollated)
+				{
+					tableCss += "td { max-height: 40px; height:40px } \n";
+				}
+				tableCss += "tr:last-child td{border-bottom: 2px solid black;} \n";
+
+			}
+			else
+			{
+				tableCss =
+					"table, th, td { border: 2px solid black; border-collapse: collapse; padding: " + cellPadding + "px; vertical-align: center; font-family: 'Verdana', Times, serif; font-size: 11px;} \n";
+			}
+
+
+			bool bOddRowsGray = !(!bAddMinMaxColumns || !isCollated);
+			tableCss += "tr:nth-child(" + (bOddRowsGray ? "odd" : "even") + ") {background-color: #e2e2e2;} \n";
+			tableCss += "tr:nth-child(" + (bOddRowsGray ? "even" : "odd") + ") {background-color: #ffffff;} \n";
+			tableCss += "tr:first-child {background-color: #ffffff;} \n";
+			tableCss += "tr.lastHeaderRow th { border-bottom: 2px solid black; } \n";
+
+			// Section start row styles
+			tableCss += "tr.sectionStartLevel0 td { border-top: 2px solid black; } \n";
+			tableCss += "tr.sectionStartLevel1 td { border-top: 1px solid black; } \n";
+			tableCss += "tr.sectionStartLevel2 td { border-top: 1px dashed black; } \n";
+
+			htmlFile.WriteLine(tableCss);
+
+			htmlFile.WriteLine("</style>");
+			htmlFile.WriteLine("</head><body>");
+			htmlFile.WriteLine("<table id='mainTable'>");
+
+			// Get format info for the columns
+			Dictionary<SummaryTableColumn, SummaryTableColumnFormatInfo> columnFormatInfoLookup = new Dictionary<SummaryTableColumn, SummaryTableColumnFormatInfo>();
+			foreach (SummaryTableColumn column in columns)
+			{
+				columnFormatInfoLookup[column] = (columnFormatInfoList != null) ? columnFormatInfoList.GetFormatInfo(column.name) : SummaryTableColumnFormatInfoCollection.DefaultColumnInfo;
+			}
+
+			// Automatically colourize the table
+			if (bScrollableTable)
+			{
+				foreach (SummaryTableColumn column in columns)
+				{
+					if (column.isNumeric)
+					{
+						column.ComputeAutomaticColourThresholds(columnFormatInfoLookup[column].autoColorizeMode);
+					}
+				}
+			}
+
+			string HeaderRow = "";
+			if (isCollated)
+			{
+				string TopHeaderRow = "";
+				if (bScrollableTable)
+				{
+					// Generate an automatic title
+					string title = htmlFilename.Replace("_Email.html", "").Replace(".html", "").Replace("\\", "/");
+					title = title.Substring(title.LastIndexOf('/') + 1);
+					TopHeaderRow += "<th colspan='" + firstStatColumnIndex + "'><h3>" + title + "</h3></th>";
+				}
+				else
+				{
+					TopHeaderRow += "<th colspan='" + firstStatColumnIndex + "'/>";
+				}
+
+				for (int i = 0; i < firstStatColumnIndex; i++)
+				{
+					HeaderRow += "<th>" + columns[i].GetDisplayName() + "</th>";
+				}
+				if (!bAddMinMaxColumns)
+				{
+					TopHeaderRow = HeaderRow;
+				}
+
+				for (int i = firstStatColumnIndex; i < columns.Count; i++)
+				{
+					string prefix = "";
+					string suffix = "";
+					string statName = GetStatNameWithPrefixAndSuffix(columns[i].GetDisplayName(), out prefix, out suffix);
+					if ((i - 1) % statColSpan == 0)
+					{
+						TopHeaderRow += "<th colspan='" + statColSpan + "' >" + statName + suffix + "</th>";
+					}
+					HeaderRow += "<th>" + prefix.Trim() + "</th>";
+				}
+				if (bAddMinMaxColumns)
+				{
+					htmlFile.WriteLine("  <tr>" + TopHeaderRow + "</tr>");
+					htmlFile.WriteLine("  <tr class='lastHeaderRow'>" + HeaderRow + "</tr>");
+				}
+				else
+				{
+					htmlFile.WriteLine("  <tr class='lastHeaderRow'>" + TopHeaderRow + "</tr>");
+				}
+			}
+			else
+			{
+				foreach (SummaryTableColumn column in columns)
+				{
+					HeaderRow += "<th>" + column.GetDisplayName() + "</th>";
+				}
+				htmlFile.WriteLine("  <tr class='lastHeaderRow'>" + HeaderRow + "</tr>");
+			}
+			string[] stripeColors = { "'#e2e2e2'", "'#ffffff'" };
+
+			// Work out which rows are major/minor section boundaries
+			Dictionary<int, int> rowSectionBoundaryLevel = new Dictionary<int, int>();
+			if (sectionBoundaries != null)
+			{
+				foreach (SummarySectionBoundaryInfo sectionBoundaryInfo in sectionBoundaries)
+				{
+					// Skip this section boundary info if it's not in this table type
+					if (isCollated && !sectionBoundaryInfo.inCollatedTable)
+					{
+						continue;
+					}
+					if (!isCollated && !sectionBoundaryInfo.inFullTable)
+					{
+						continue;
+					}
+					string prevSectionName = "";
+					for (int i = 0; i < rowCount; i++)
+					{
+						int boundaryLevel = 0;
+						if (sectionBoundaryInfo != null)
+						{
+							// Work out the section name if we have section boundary info. When it changes, apply the sectionStart CSS class
+							string sectionName = "";
+							if (sectionBoundaryInfo != null && columnLookup.ContainsKey(sectionBoundaryInfo.statName))
+							{
+								// Get the section name
+								if (!columnLookup.ContainsKey(sectionBoundaryInfo.statName))
+								{
+									continue;
+								}
+								SummaryTableColumn col = columnLookup[sectionBoundaryInfo.statName];
+								sectionName = col.GetStringValue(i);
+
+								// if we have a start token then strip before it
+								if (sectionBoundaryInfo.startToken != null)
+								{
+									int startTokenIndex = sectionName.IndexOf(sectionBoundaryInfo.startToken);
+									if (startTokenIndex != -1)
+									{
+										sectionName = sectionName.Substring(startTokenIndex + sectionBoundaryInfo.startToken.Length);
+									}
+								}
+
+								// if we have an end token then strip after it
+								if (sectionBoundaryInfo.endToken != null)
+								{
+									int endTokenIndex = sectionName.IndexOf(sectionBoundaryInfo.endToken);
+									if (endTokenIndex != -1)
+									{
+										sectionName = sectionName.Substring(0, endTokenIndex);
+									}
+								}
+							}
+							if (sectionName != prevSectionName && i > 0)
+							{
+								// Update the row's boundary type info
+								boundaryLevel = sectionBoundaryInfo.level;
+								if (rowSectionBoundaryLevel.ContainsKey(i))
+								{
+									// Lower level values override higher ones
+									boundaryLevel = Math.Min(rowSectionBoundaryLevel[i], boundaryLevel);
+								}
+								rowSectionBoundaryLevel[i] = boundaryLevel;
+							}
+							prevSectionName = sectionName;
+						}
+					}
+				}
+
+			}
+
+			// Add the rows to the table
+			for (int i = 0; i < rowCount; i++)
+			{
+				string rowClassStr = "";
+
+				// Is this a major/minor section boundary
+				if (rowSectionBoundaryLevel.ContainsKey(i))
+				{
+					int sectionLevel = rowSectionBoundaryLevel[i];
+					if (sectionLevel < 3)
+					{
+						rowClassStr = " class='sectionStartLevel" + sectionLevel + "'";
+					}
+				}
+
+				htmlFile.Write("<tr" + rowClassStr + ">");
+				int columnIndex = 0;
+				foreach (SummaryTableColumn column in columns)
+				{
+					// Add the tooltip for non-collated tables
+					string toolTipString = "";
+					if (!isCollated)
+					{
+						string toolTip = column.GetToolTipValue(i);
+						if (toolTip == "")
+						{
+							toolTip = column.GetDisplayName();
+						}
+						toolTipString = " title='" + toolTip + "'";
+					}
+					string colour = column.GetColour(i);
+
+					// Alternating row colours are normally handled by CSS, but we need to handle it explicitly if we have frozen first columns
+					if (columnIndex < stickyColumnCount && colour == null)
+					{
+						colour = stripeColors[i % 2];
+					}
+					string bgColorString = (colour == null ? "" : " bgcolor=" + colour);
+					bool bold = false;
+					string numericFormat = columnFormatInfoLookup[column].numericFormat;
+					string stringValue = column.GetStringValue(i, true, numericFormat);
+					if (maxColumnStringLength > 0 && stringValue.Length > maxColumnStringLength)
+					{
+						stringValue = TableUtil.SafeTruncateHtmlTableValue(stringValue, maxColumnStringLength);
+					}
+					if (bSpreadsheetFriendlyStrings && !column.isNumeric)
+					{
+						stringValue = "'" + stringValue;
+					}
+					string columnString = "<td" + toolTipString + bgColorString + "> " + (bold ? "<b>" : "") + stringValue + (bold ? "</b>" : "") + "</td>";
+					htmlFile.Write(columnString);
+					columnIndex++;
+				}
+				htmlFile.WriteLine("</tr>");
+			}
+			htmlFile.WriteLine("</table>");
+			string extraString = "";
+			if (isCollated && weightByColumnName != null)
+			{
+				extraString += " - weighted avg";
+				//htmlFile.WriteLine("<p style='font-size:8'>Weighted by " + weightByColumnName +"</p>");
+			}
+
+			htmlFile.WriteLine("<p style='font-size:8'>Created with PerfReportTool " + VersionString + extraString + "</p>");
+			htmlFile.WriteLine("</font></body></html>");
+
+			htmlFile.Close();
+		}
+
+		public SummaryTable SortRows(List<string> rowSortList, bool reverseSort)
+		{
+			List<KeyValuePair<string, int>> columnRemapping = new List<KeyValuePair<string, int>>();
+			for (int i = 0; i < rowCount; i++)
+			{
+				string key = "";
+				foreach (string s in rowSortList)
+				{
+					if (columnLookup.ContainsKey(s.ToLower()))
+					{
+						SummaryTableColumn column = columnLookup[s.ToLower()];
+						key += "{" + column.GetStringValue(i) + "}";
+					}
+					else
+					{
+						key += "{}";
+					}
+				}
+				columnRemapping.Add(new KeyValuePair<string, int>(key, i));
+			}
+
+			columnRemapping.Sort(delegate (KeyValuePair<string, int> m1, KeyValuePair<string, int> m2)
+			{
+				return m1.Key.CompareTo(m2.Key);
+			});
+
+			// Reorder the metadata rows
+			List<SummaryTableColumn> newColumns = new List<SummaryTableColumn>();
+			foreach (SummaryTableColumn srcCol in columns)
+			{
+				SummaryTableColumn destCol = new SummaryTableColumn(srcCol.name, srcCol.isNumeric);
+				for (int i = 0; i < rowCount; i++)
+				{
+					int srcIndex = columnRemapping[i].Value;
+					int destIndex = reverseSort ? rowCount - 1 - i : i;
+					if (srcCol.isNumeric)
+					{
+						destCol.SetValue(destIndex, srcCol.GetValue(srcIndex));
+					}
+					else
+					{
+						destCol.SetStringValue(destIndex, srcCol.GetStringValue(srcIndex));
+					}
+					destCol.SetColourThresholds(destIndex, srcCol.GetColourThresholds(srcIndex));
+					destCol.SetToolTipValue(destIndex, srcCol.GetToolTipValue(srcIndex));
+				}
+				newColumns.Add(destCol);
+			}
+			SummaryTable newTable = new SummaryTable();
+			newTable.columns = newColumns;
+			newTable.rowCount = rowCount;
+			newTable.firstStatColumnIndex = firstStatColumnIndex;
+			newTable.isCollated = isCollated;
+			newTable.InitColumnLookup();
+			return newTable;
+		}
+
+		void InitColumnLookup()
+		{
+			columnLookup.Clear();
+			foreach (SummaryTableColumn col in columns)
+			{
+				columnLookup.Add(col.name.ToLower(), col);
+			}
+		}
+
+		public void AddRowData(SummaryTableRowData metadata, bool bIncludeCsvStatAverages, bool bIncludeHiddenStats)
+		{
+			foreach (string key in metadata.dict.Keys)
+			{
+				SummaryTableElement value = metadata.dict[key];
+				if (value.type == SummaryTableElement.Type.CsvStatAverage && !bIncludeCsvStatAverages)
+				{
+					continue;
+				}
+				if (value.GetFlag(SummaryTableElement.Flags.Hidden) && !bIncludeHiddenStats)
+				{
+					continue;
+				}
+				SummaryTableColumn column = null;
+
+				if (!columnLookup.ContainsKey(key))
+				{
+					column = new SummaryTableColumn(value.name, value.isNumeric);
+					columnLookup.Add(key, column);
+					columns.Add(column);
+				}
+				else
+				{
+					column = columnLookup[key];
+				}
+
+				if (value.isNumeric)
+				{
+					column.SetValue(rowCount, (double)value.numericValue);
+				}
+				else
+				{
+					column.SetStringValue(rowCount, value.value);
+				}
+				column.SetColourThresholds(rowCount, value.colorThresholdList);
+				column.SetToolTipValue(rowCount, value.tooltip);
+			}
+			rowCount++;
+		}
+
+		public int Count
+		{
+			get { return rowCount; }
+		}
+
+		Dictionary<string, SummaryTableColumn> columnLookup = new Dictionary<string, SummaryTableColumn>();
+		List<SummaryTableColumn> columns = new List<SummaryTableColumn>();
+		List<double> rowWeightings = null;
+		int rowCount = 0;
+		int firstStatColumnIndex = 0;
+		bool isCollated = false;
+		bool hasMinMaxColumns = false;
+	};
+
+}
\ No newline at end of file
diff --git a/Engine/Source/Programs/CSVTools/PerfReportTool/SummaryTypes.cs b/Engine/Source/Programs/CSVTools/PerfReportTool/SummaryTypes.cs
new file mode 100644
--- /dev/null
+++ b/Engine/Source/Programs/CSVTools/PerfReportTool/SummaryTypes.cs
@@ -0,0 +1,180 @@
+// Copyright Epic Games, Inc. All Rights Reserved.
+
+using System;
+using System.Collections.Generic;
+
+namespace PerfSummaries
+{
+	class Colour
+	{
+		public Colour(string str)
+		{
+			string hexStr = str.TrimStart('#');
+			int hexValue = Convert.ToInt32(hexStr, 16);
+			byte rb = (byte)((hexValue >> 16) & 0xff);
+			byte gb = (byte)((hexValue >> 8) & 0xff);
+			byte bb = (byte)((hexValue >> 0) & 0xff);
+			r = ((float)rb) / 255.0f;
+			g = ((float)gb) / 255.0f;
+			b = ((float)bb) / 255.0f;
+			alpha = 1.0f;
+		}
+		public Colour(uint hex, float alphaIn = 1.0f)
+		{
+			byte rb = (byte)((hex >> 16) & 0xff);
+			byte gb = (byte)((hex >> 8) & 0xff);
+			byte bb = (byte)((hex >> 0) & 0xff);
+
+			r = ((float)rb) / 255.0f;
+			g = ((float)rb) / 255.0f;
+			b = ((float)rb) / 255.0f;
+
+			alpha = alphaIn;
+		}
+		public Colour(Colour colourIn) { r = colourIn.r; g = colourIn.g; b = colourIn.b; alpha = colourIn.alpha; }
+		public Colour(float rIn, float gIn, float bIn, float aIn = 1.0f) { r = rIn; g = gIn; b = bIn; alpha = aIn; }
+
+		public static Colour Lerp(Colour Colour0, Colour Colour1, float t)
+		{
+			return new Colour(
+				Colour0.r * (1.0f - t) + Colour1.r * t,
+				Colour0.g * (1.0f - t) + Colour1.g * t,
+				Colour0.b * (1.0f - t) + Colour1.b * t,
+				Colour0.alpha * (1.0f - t) + Colour1.alpha * t);
+
+		}
+
+		public string ToHTMLString()
+		{
+			return "'" + ToString() + "'";
+		}
+
+		public override string ToString()
+		{
+			int rI = (int)(r * 255.0f);
+			int gI = (int)(g * 255.0f);
+			int bI = (int)(b * 255.0f);
+			return "#" + rI.ToString("x2") + gI.ToString("x2") + bI.ToString("x2");
+		}
+
+		public static Colour White = new Colour(1.0f, 1.0f, 1.0f, 1.0f);
+		public static Colour Black = new Colour(0, 0, 0, 1.0f);
+		public static Colour Orange = new Colour(1.0f, 0.5f, 0.0f, 1.0f);
+		public static Colour Yellow = new Colour(1.0f, 1.0f, 0.0f, 1.0f);
+		public static Colour Red = new Colour(1.0f, 0.0f, 0.0f, 1.0f);
+		public static Colour Green = new Colour(0.0f, 1.0f, 0.0f, 1.0f);
+
+		public float r, g, b;
+		public float alpha;
+	};
+
+
+	class ThresholdInfo
+	{
+		public ThresholdInfo(double inValue, Colour inColourOverride = null)
+		{
+			value = inValue;
+			colour = inColourOverride;
+		}
+		public double value;
+		public Colour colour;
+	};
+
+	class ColourThresholdList
+	{
+		public static string GetThresholdColour(double value, double redValue, double orangeValue, double yellowValue, double greenValue,
+			Colour redOverride = null, Colour orangeOverride = null, Colour yellowOverride = null, Colour greenOverride = null)
+		{
+			Colour green = (greenOverride != null) ? greenOverride : new Colour(0.0f, 1.0f, 0.0f, 1.0f);
+			Colour orange = (orangeOverride != null) ? orangeOverride : new Colour(1.0f, 0.5f, 0.0f, 1.0f);
+			Colour yellow = (yellowOverride != null) ? yellowOverride : new Colour(1.0f, 1.0f, 0.0f, 1.0f);
+			Colour red = (redOverride != null) ? redOverride : new Colour(1.0f, 0.0f, 0.0f, 1.0f);
+
+			if (redValue > orangeValue)
+			{
+				redValue = -redValue;
+				orangeValue = -orangeValue;
+				yellowValue = -yellowValue;
+				greenValue = -greenValue;
+				value = -value;
+			}
+
+			Colour col = null;
+			if (value <= redValue)
+			{
+				col = red;
+			}
+			else if (value <= orangeValue)
+			{
+				double t = (value - redValue) / (orangeValue - redValue);
+				col = Colour.Lerp(red, orange, (float)t);
+			}
+			else if (value <= yellowValue)
+			{
+				double t = (value - orangeValue) / (yellowValue - orangeValue);
+				col = Colour.Lerp(orange, yellow, (float)t);
+			}
+			else if (value <= greenValue)
+			{
+				float t = (float)(value - yellowValue) / (float)(greenValue - yellowValue);
+				col = Colour.Lerp(yellow, green, t);
+			}
+			else
+			{
+				col = green;
+			}
+			return col.ToHTMLString();
+		}
+
+		public void Add(ThresholdInfo info)
+		{
+			if (Thresholds.Count < 4)
+			{
+				Thresholds.Add(info);
+			}
+		}
+		public int Count
+		{
+			get { return Thresholds.Count; }
+		}
+
+		public string GetColourForValue(string value)
+		{
+			try
+			{
+				return GetColourForValue(Convert.ToDouble(value, System.Globalization.CultureInfo.InvariantCulture));
+			}
+			catch
+			{
+				return "'#ffffff'";
+			}
+		}
+
+		public string GetColourForValue(double value)
+		{
+			if (Thresholds.Count == 4)
+			{
+				return GetThresholdColour(value, Thresholds[3].value, Thresholds[2].value, Thresholds[1].value, Thresholds[0].value, Thresholds[3].colour, Thresholds[2].colour, Thresholds[1].colour, Thresholds[0].colour);
+			}
+			return "'#ffffff'";
+		}
+
+		public static string GetSafeColourForValue(ColourThresholdList list, string value)
+		{
+			if (list == null)
+			{
+				return "'#ffffff'";
+			}
+			return list.GetColourForValue(value);
+		}
+		public static string GetSafeColourForValue(ColourThresholdList list, double value)
+		{
+			if (list == null)
+			{
+				return "'#ffffff'";
+			}
+			return list.GetColourForValue(value);
+		}
+		public List<ThresholdInfo> Thresholds = new List<ThresholdInfo>();
+	};
+}
\ No newline at end of file
diff --git a/Engine/Source/Programs/CSVTools/PerfReportTool/perfsummaries.cs b/Engine/Source/Programs/CSVTools/PerfReportTool/perfsummaries.cs
deleted file mode 100644
