Move SummaryTableElement and SummaryTableRowData into separate SummaryTableData.cs file.
This moves core data and summary table cache serialization away from the higher level summary table logic

diff --git a/Engine/Source/Programs/CSVTools/PerfReportTool/PerfreportTool.csproj b/Engine/Source/Programs/CSVTools/PerfReportTool/PerfreportTool.csproj
--- a/Engine/Source/Programs/CSVTools/PerfReportTool/PerfreportTool.csproj
+++ b/Engine/Source/Programs/CSVTools/PerfReportTool/PerfreportTool.csproj
@@ -44,6 +44,7 @@
   <ItemGroup>
     <Compile Include="CsvFileCache..cs" />
     <Compile Include="PerfReportTool.cs" />
+    <Compile Include="SummaryTableData.cs" />
     <Compile Include="SummaryTypes.cs" />
     <Compile Include="Properties\AssemblyInfo.cs" />
     <Compile Include="ReportGraph.cs" />
diff --git a/Engine/Source/Programs/CSVTools/PerfReportTool/SummaryTable.cs b/Engine/Source/Programs/CSVTools/PerfReportTool/SummaryTable.cs
--- a/Engine/Source/Programs/CSVTools/PerfReportTool/SummaryTable.cs
+++ b/Engine/Source/Programs/CSVTools/PerfReportTool/SummaryTable.cs
@@ -43,300 +43,6 @@ namespace PerfSummaries
 		}
 	}
 
-	class SummaryTableElement
-	{
-		// Bump this when making changes!
-		public static int CacheVersion = 1;
-
-		// NOTE: this is serialized. Don't change the order!
-		public enum Type
-		{
-			CsvStatAverage,
-			CsvMetadata,
-			SummaryTableMetric,
-			ToolMetadata
-		};
-
-		public enum Flags
-		{
-			Hidden = 0x01
-		};
-
-		private SummaryTableElement()
-		{
-
-		}
-		public SummaryTableElement(Type inType, string inName, double inValue, ColourThresholdList inColorThresholdList, string inToolTip, uint inFlags = 0)
-		{
-			type = inType;
-			name = inName;
-			isNumeric = true;
-			numericValue = inValue;
-			value = inValue.ToString();
-			colorThresholdList = inColorThresholdList;
-			tooltip = inToolTip;
-			flags = inFlags;
-		}
-		public SummaryTableElement(Type inType, string inName, string inValue, ColourThresholdList inColorThresholdList, string inToolTip, uint inFlags = 0)
-		{
-			type = inType;
-			name = inName;
-			numericValue = 0.0;
-			isNumeric = false;
-			colorThresholdList = inColorThresholdList;
-			value = inValue;
-			tooltip = inToolTip;
-			flags = inFlags;
-		}
-
-		public static SummaryTableElement ReadFromCache(BinaryReader reader)
-		{
-			SummaryTableElement val = new SummaryTableElement();
-			val.type = (Type)reader.ReadUInt32();
-			val.name = reader.ReadString();
-			val.value = reader.ReadString();
-			val.tooltip = reader.ReadString();
-			val.numericValue = reader.ReadDouble();
-			val.isNumeric = reader.ReadBoolean();
-			val.flags = reader.ReadUInt32();
-			bool hasThresholdList = reader.ReadBoolean();
-			if (hasThresholdList)
-			{
-				int thresholdCount = reader.ReadInt32();
-				val.colorThresholdList = new ColourThresholdList();
-				for (int i = 0; i < thresholdCount; i++)
-				{
-					bool bHasColour = reader.ReadBoolean();
-					Colour thresholdColour = null;
-					if (bHasColour)
-					{
-						thresholdColour = new Colour(reader.ReadString());
-					}
-					double thresholdValue = reader.ReadDouble();
-					ThresholdInfo info = new ThresholdInfo(thresholdValue, thresholdColour);
-					val.colorThresholdList.Add(info);
-				}
-			}
-			return val;
-		}
-
-
-
-		public void WriteToCache(BinaryWriter writer)
-		{
-			writer.Write((uint)type);
-			writer.Write(name);
-			writer.Write(value);
-			writer.Write(tooltip);
-			writer.Write(numericValue);
-			writer.Write(isNumeric);
-			writer.Write(flags);
-			writer.Write(colorThresholdList != null);
-			if (colorThresholdList != null)
-			{
-				writer.Write((int)colorThresholdList.Count);
-				foreach (ThresholdInfo thresholdInfo in colorThresholdList.Thresholds)
-				{
-					writer.Write(thresholdInfo.colour != null);
-					if (thresholdInfo.colour != null)
-					{
-						writer.Write(thresholdInfo.colour.ToString());
-					}
-					writer.Write(thresholdInfo.value);
-				}
-			}
-		}
-
-		public SummaryTableElement Clone()
-		{
-			return (SummaryTableElement)MemberwiseClone();
-		}
-
-		public void SetFlag(Flags flag, bool value)
-		{
-			if (value)
-			{
-				flags |= (uint)flag;
-			}
-			else
-			{
-				flags &= ~(uint)flag;
-			}
-		}
-		public bool GetFlag(Flags flag)
-		{
-			return (flags & (uint)flag) != 0;
-		}
-
-		public Type type;
-		public string name;
-		public string value;
-		public string tooltip;
-		public ColourThresholdList colorThresholdList;
-		public double numericValue;
-		public bool isNumeric;
-		public uint flags;
-	}
-	class SummaryTableRowData
-	{
-		public SummaryTableRowData()
-		{
-		}
-
-		// TODO: If this is bumped beyond 6, we need to implement backwards compatibility
-		static int CacheVersion = 6;
-
-		public static SummaryTableRowData TryReadFromCache(string summaryTableCacheDir, string csvId)
-		{
-			string filename = Path.Combine(summaryTableCacheDir, csvId + ".prc");
-			return TryReadFromCacheFile(filename);
-		}
-
-		public static SummaryTableRowData TryReadFromCacheFile(string filename, bool bReadJustInitialMetadata = false)
-		{
-			SummaryTableRowData metaData = null;
-			if (!File.Exists(filename))
-			{
-				return null;
-			}
-			try
-			{
-				using (FileStream fileStream = new FileStream(filename, FileMode.Open, FileAccess.Read))
-				{
-					BinaryReader reader = new BinaryReader(fileStream);
-					int version = reader.ReadInt32();
-					int metadataValueVersion = reader.ReadInt32();
-					if (version == CacheVersion && metadataValueVersion == SummaryTableElement.CacheVersion)
-					{
-						bool bEarlyOut = false;
-						metaData = new SummaryTableRowData();
-						int dictEntryCount = reader.ReadInt32();
-						for (int i = 0; i < dictEntryCount; i++)
-						{
-							string key = reader.ReadString();
-							SummaryTableElement value = SummaryTableElement.ReadFromCache(reader);
-							// If we're just reading initial metadata then skip everything after ToolMetadata and CsvMetadata
-							if (bReadJustInitialMetadata && value.type != SummaryTableElement.Type.ToolMetadata && value.type != SummaryTableElement.Type.CsvMetadata)
-							{
-								bEarlyOut = true;
-								break;
-							}
-							metaData.dict.Add(key, value);
-						}
-
-						if (!bEarlyOut)
-						{
-							string endString = reader.ReadString();
-							if (endString != "END")
-							{
-								Console.WriteLine("Corruption detected in " + filename + ". Skipping read");
-								metaData = null;
-							}
-						}
-					}
-					reader.Close();
-				}
-			}
-			catch (Exception e)
-			{
-				metaData = null;
-				Console.WriteLine("Error reading from cache file " + filename + ": " + e.Message);
-			}
-			return metaData;
-		}
-
-		public bool WriteToCache(string summaryTableCacheDir, string csvId)
-		{
-			string filename = Path.Combine(summaryTableCacheDir, csvId + ".prc");
-			try
-			{
-				using (FileStream fileStream = new FileStream(filename, FileMode.Create))
-				{
-					BinaryWriter writer = new BinaryWriter(fileStream);
-					writer.Write(CacheVersion);
-					writer.Write(SummaryTableElement.CacheVersion);
-
-					writer.Write(dict.Count);
-					foreach (KeyValuePair<string, SummaryTableElement> entry in dict)
-					{
-						writer.Write(entry.Key);
-						entry.Value.WriteToCache(writer);
-					}
-					writer.Write("END");
-					writer.Close();
-				}
-			}
-			catch (IOException)
-			{
-				Console.WriteLine("Failed to write to cache file " + filename + ".");
-				return false;
-			}
-			return true;
-		}
-
-		public void RemoveSafe(string name)
-		{
-			string key = name.ToLower();
-			if (dict.ContainsKey(key))
-			{
-				dict.Remove(key);
-			}
-		}
-
-		public void Add(SummaryTableElement.Type type, string name, double value, ColourThresholdList colorThresholdList = null, string tooltip = "", uint flags = 0)
-		{
-			string key = name.ToLower();
-			SummaryTableElement metadataValue = new SummaryTableElement(type, name, value, colorThresholdList, tooltip, flags);
-			try
-			{
-				dict.Add(key, metadataValue);
-			}
-			catch (System.ArgumentException)
-			{
-				throw new Exception("Summary metadata key " + key + " has already been added");
-			}
-		}
-
-		public void Add(SummaryTableElement.Type type, string name, string value, ColourThresholdList colorThresholdList = null, string tooltip = "", uint flags = 0)
-		{
-			string key = name.ToLower();
-			double numericValue = double.MaxValue;
-			try
-			{
-				numericValue = Convert.ToDouble(value, System.Globalization.CultureInfo.InvariantCulture);
-			}
-			catch { }
-
-			SummaryTableElement metadataValue = null;
-			if (numericValue != double.MaxValue)
-			{
-				metadataValue = new SummaryTableElement(type, name, numericValue, colorThresholdList, tooltip, flags);
-			}
-			else
-			{
-				metadataValue = new SummaryTableElement(type, name, value, colorThresholdList, tooltip, flags);
-			}
-
-			try
-			{
-				dict.Add(key, metadataValue);
-			}
-			catch (System.ArgumentException)
-			{
-				throw new Exception("Summary metadata key " + key + " has already been added");
-			}
-		}
-
-		public void AddString(SummaryTableElement.Type type, string name, string value, ColourThresholdList colorThresholdList = null, string tooltip = "")
-		{
-			string key = name.ToLower();
-			SummaryTableElement metadataValue = new SummaryTableElement(type, name, value, colorThresholdList, tooltip);
-			dict.Add(key, metadataValue);
-		}
-
-		public Dictionary<string, SummaryTableElement> dict = new Dictionary<string, SummaryTableElement>();
-	};
-
 	class SummarySectionBoundaryInfo
 	{
 		public SummarySectionBoundaryInfo(string inStatName, string inStartToken, string inEndToken, int inLevel, bool inInCollatedTable, bool inInFullTable)
diff --git a/Engine/Source/Programs/CSVTools/PerfReportTool/SummaryTableData.cs b/Engine/Source/Programs/CSVTools/PerfReportTool/SummaryTableData.cs
new file mode 100644
--- /dev/null
+++ b/Engine/Source/Programs/CSVTools/PerfReportTool/SummaryTableData.cs
@@ -0,0 +1,306 @@
+// Copyright Epic Games, Inc. All Rights Reserved.
+
+using System;
+using System.Collections.Generic;
+using System.Linq;
+using System.Xml.Linq;
+using PerfReportTool;
+using System.IO;
+
+namespace PerfSummaries
+{
+	class SummaryTableElement
+	{
+		// Bump this when making changes!
+		public static int CacheVersion = 1;
+
+		// NOTE: this is serialized. Don't change the order!
+		public enum Type
+		{
+			CsvStatAverage,
+			CsvMetadata,
+			SummaryTableMetric,
+			ToolMetadata
+		};
+
+		public enum Flags
+		{
+			Hidden = 0x01
+		};
+
+		private SummaryTableElement()
+		{
+
+		}
+		public SummaryTableElement(Type inType, string inName, double inValue, ColourThresholdList inColorThresholdList, string inToolTip, uint inFlags = 0)
+		{
+			type = inType;
+			name = inName;
+			isNumeric = true;
+			numericValue = inValue;
+			value = inValue.ToString();
+			colorThresholdList = inColorThresholdList;
+			tooltip = inToolTip;
+			flags = inFlags;
+		}
+		public SummaryTableElement(Type inType, string inName, string inValue, ColourThresholdList inColorThresholdList, string inToolTip, uint inFlags = 0)
+		{
+			type = inType;
+			name = inName;
+			numericValue = 0.0;
+			isNumeric = false;
+			colorThresholdList = inColorThresholdList;
+			value = inValue;
+			tooltip = inToolTip;
+			flags = inFlags;
+		}
+
+		public static SummaryTableElement ReadFromCache(BinaryReader reader)
+		{
+			SummaryTableElement val = new SummaryTableElement();
+			val.type = (Type)reader.ReadUInt32();
+			val.name = reader.ReadString();
+			val.value = reader.ReadString();
+			val.tooltip = reader.ReadString();
+			val.numericValue = reader.ReadDouble();
+			val.isNumeric = reader.ReadBoolean();
+			val.flags = reader.ReadUInt32();
+			bool hasThresholdList = reader.ReadBoolean();
+			if (hasThresholdList)
+			{
+				int thresholdCount = reader.ReadInt32();
+				val.colorThresholdList = new ColourThresholdList();
+				for (int i = 0; i < thresholdCount; i++)
+				{
+					bool bHasColour = reader.ReadBoolean();
+					Colour thresholdColour = null;
+					if (bHasColour)
+					{
+						thresholdColour = new Colour(reader.ReadString());
+					}
+					double thresholdValue = reader.ReadDouble();
+					ThresholdInfo info = new ThresholdInfo(thresholdValue, thresholdColour);
+					val.colorThresholdList.Add(info);
+				}
+			}
+			return val;
+		}
+
+
+
+		public void WriteToCache(BinaryWriter writer)
+		{
+			writer.Write((uint)type);
+			writer.Write(name);
+			writer.Write(value);
+			writer.Write(tooltip);
+			writer.Write(numericValue);
+			writer.Write(isNumeric);
+			writer.Write(flags);
+			writer.Write(colorThresholdList != null);
+			if (colorThresholdList != null)
+			{
+				writer.Write((int)colorThresholdList.Count);
+				foreach (ThresholdInfo thresholdInfo in colorThresholdList.Thresholds)
+				{
+					writer.Write(thresholdInfo.colour != null);
+					if (thresholdInfo.colour != null)
+					{
+						writer.Write(thresholdInfo.colour.ToString());
+					}
+					writer.Write(thresholdInfo.value);
+				}
+			}
+		}
+
+		public SummaryTableElement Clone()
+		{
+			return (SummaryTableElement)MemberwiseClone();
+		}
+
+		public void SetFlag(Flags flag, bool value)
+		{
+			if (value)
+			{
+				flags |= (uint)flag;
+			}
+			else
+			{
+				flags &= ~(uint)flag;
+			}
+		}
+		public bool GetFlag(Flags flag)
+		{
+			return (flags & (uint)flag) != 0;
+		}
+
+		public Type type;
+		public string name;
+		public string value;
+		public string tooltip;
+		public ColourThresholdList colorThresholdList;
+		public double numericValue;
+		public bool isNumeric;
+		public uint flags;
+	}
+	class SummaryTableRowData
+	{
+		public SummaryTableRowData()
+		{
+		}
+
+		// TODO: If this is bumped beyond 6, we need to implement backwards compatibility
+		static int CacheVersion = 6;
+
+		public static SummaryTableRowData TryReadFromCache(string summaryTableCacheDir, string csvId)
+		{
+			string filename = Path.Combine(summaryTableCacheDir, csvId + ".prc");
+			return TryReadFromCacheFile(filename);
+		}
+
+		public static SummaryTableRowData TryReadFromCacheFile(string filename, bool bReadJustInitialMetadata = false)
+		{
+			SummaryTableRowData metaData = null;
+			if (!File.Exists(filename))
+			{
+				return null;
+			}
+			try
+			{
+				using (FileStream fileStream = new FileStream(filename, FileMode.Open, FileAccess.Read))
+				{
+					BinaryReader reader = new BinaryReader(fileStream);
+					int version = reader.ReadInt32();
+					int metadataValueVersion = reader.ReadInt32();
+					if (version == CacheVersion && metadataValueVersion == SummaryTableElement.CacheVersion)
+					{
+						bool bEarlyOut = false;
+						metaData = new SummaryTableRowData();
+						int dictEntryCount = reader.ReadInt32();
+						for (int i = 0; i < dictEntryCount; i++)
+						{
+							string key = reader.ReadString();
+							SummaryTableElement value = SummaryTableElement.ReadFromCache(reader);
+							// If we're just reading initial metadata then skip everything after ToolMetadata and CsvMetadata
+							if (bReadJustInitialMetadata && value.type != SummaryTableElement.Type.ToolMetadata && value.type != SummaryTableElement.Type.CsvMetadata)
+							{
+								bEarlyOut = true;
+								break;
+							}
+							metaData.dict.Add(key, value);
+						}
+
+						if (!bEarlyOut)
+						{
+							string endString = reader.ReadString();
+							if (endString != "END")
+							{
+								Console.WriteLine("Corruption detected in " + filename + ". Skipping read");
+								metaData = null;
+							}
+						}
+					}
+					reader.Close();
+				}
+			}
+			catch (Exception e)
+			{
+				metaData = null;
+				Console.WriteLine("Error reading from cache file " + filename + ": " + e.Message);
+			}
+			return metaData;
+		}
+
+		public bool WriteToCache(string summaryTableCacheDir, string csvId)
+		{
+			string filename = Path.Combine(summaryTableCacheDir, csvId + ".prc");
+			try
+			{
+				using (FileStream fileStream = new FileStream(filename, FileMode.Create))
+				{
+					BinaryWriter writer = new BinaryWriter(fileStream);
+					writer.Write(CacheVersion);
+					writer.Write(SummaryTableElement.CacheVersion);
+
+					writer.Write(dict.Count);
+					foreach (KeyValuePair<string, SummaryTableElement> entry in dict)
+					{
+						writer.Write(entry.Key);
+						entry.Value.WriteToCache(writer);
+					}
+					writer.Write("END");
+					writer.Close();
+				}
+			}
+			catch (IOException)
+			{
+				Console.WriteLine("Failed to write to cache file " + filename + ".");
+				return false;
+			}
+			return true;
+		}
+
+		public void RemoveSafe(string name)
+		{
+			string key = name.ToLower();
+			if (dict.ContainsKey(key))
+			{
+				dict.Remove(key);
+			}
+		}
+
+		public void Add(SummaryTableElement.Type type, string name, double value, ColourThresholdList colorThresholdList = null, string tooltip = "", uint flags = 0)
+		{
+			string key = name.ToLower();
+			SummaryTableElement metadataValue = new SummaryTableElement(type, name, value, colorThresholdList, tooltip, flags);
+			try
+			{
+				dict.Add(key, metadataValue);
+			}
+			catch (System.ArgumentException)
+			{
+				throw new Exception("Summary metadata key " + key + " has already been added");
+			}
+		}
+
+		public void Add(SummaryTableElement.Type type, string name, string value, ColourThresholdList colorThresholdList = null, string tooltip = "", uint flags = 0)
+		{
+			string key = name.ToLower();
+			double numericValue = double.MaxValue;
+			try
+			{
+				numericValue = Convert.ToDouble(value, System.Globalization.CultureInfo.InvariantCulture);
+			}
+			catch { }
+
+			SummaryTableElement metadataValue = null;
+			if (numericValue != double.MaxValue)
+			{
+				metadataValue = new SummaryTableElement(type, name, numericValue, colorThresholdList, tooltip, flags);
+			}
+			else
+			{
+				metadataValue = new SummaryTableElement(type, name, value, colorThresholdList, tooltip, flags);
+			}
+
+			try
+			{
+				dict.Add(key, metadataValue);
+			}
+			catch (System.ArgumentException)
+			{
+				throw new Exception("Summary metadata key " + key + " has already been added");
+			}
+		}
+
+		public void AddString(SummaryTableElement.Type type, string name, string value, ColourThresholdList colorThresholdList = null, string tooltip = "")
+		{
+			string key = name.ToLower();
+			SummaryTableElement metadataValue = new SummaryTableElement(type, name, value, colorThresholdList, tooltip);
+			dict.Add(key, metadataValue);
+		}
+
+		public Dictionary<string, SummaryTableElement> dict = new Dictionary<string, SummaryTableElement>();
+	};
+
+}
\ No newline at end of file
