PerfReportTool 4.61 - Reset attributes on map summary files after copying. This ensures the file is not readonly, even if the source file is

diff --git a/Engine/Source/Programs/CSVTools/PerfReportTool/PerfReportTool.cs b/Engine/Source/Programs/CSVTools/PerfReportTool/PerfReportTool.cs
--- a/Engine/Source/Programs/CSVTools/PerfReportTool/PerfReportTool.cs
+++ b/Engine/Source/Programs/CSVTools/PerfReportTool/PerfReportTool.cs
@@ -18,7 +18,7 @@ namespace PerfReportTool
 {
     class Version
     {
-        private static string VersionString = "4.60";
+        private static string VersionString = "4.61";
 
         public static string Get() { return VersionString; }
     };
diff --git a/Engine/Source/Programs/CSVTools/PerfReportTool/Summaries/MapOverlaySummary.cs b/Engine/Source/Programs/CSVTools/PerfReportTool/Summaries/MapOverlaySummary.cs
--- a/Engine/Source/Programs/CSVTools/PerfReportTool/Summaries/MapOverlaySummary.cs
+++ b/Engine/Source/Programs/CSVTools/PerfReportTool/Summaries/MapOverlaySummary.cs
@@ -1,6 +1,7 @@
 // Copyright Epic Games, Inc. All Rights Reserved.
 
 using System;
+using System.IO;
 using System.Collections.Generic;
 using System.Linq;
 using System.Xml.Linq;
@@ -123,9 +124,11 @@ namespace PerfSummaries
 				string outputDirectory = System.IO.Path.GetDirectoryName(System.IO.Path.GetFullPath(htmlFileName));
 				string outputMapFilename = System.IO.Path.Combine(outputDirectory, destImageFilename);
 
-				if (!System.IO.File.Exists(outputMapFilename))
+				if (!File.Exists(outputMapFilename))
 				{
-					System.IO.File.Copy(sourceImagePath, outputMapFilename);
+					// Copy the file to the reports folder and reset attributes to ensure it's not readonly if the source is
+					File.Copy(sourceImagePath, outputMapFilename);
+					File.SetAttributes(outputMapFilename, FileAttributes.Normal);
 				}
 
 				// Check if the file exists in the output directory
