PerfReportTool 4.62 - Make scrollableFormatting and reverseSortRows properties of the table, settable in XML. Commandline overrides still work
- Modify defined tables to use the new attributes

diff --git a/Engine/Source/Programs/CSVTools/PerfReportTool/PerfReportTool.cs b/Engine/Source/Programs/CSVTools/PerfReportTool/PerfReportTool.cs
--- a/Engine/Source/Programs/CSVTools/PerfReportTool/PerfReportTool.cs
+++ b/Engine/Source/Programs/CSVTools/PerfReportTool/PerfReportTool.cs
@@ -18,7 +18,7 @@ namespace PerfReportTool
 {
     class Version
     {
-        private static string VersionString = "4.61";
+        private static string VersionString = "4.62";
 
         public static string Get() { return VersionString; }
     };
@@ -576,29 +576,32 @@ namespace PerfReportTool
             perfLog.LogTotalTiming();
         }
 
-        void WriteSummaryTableReport(string outputDir, string filenameWithoutExtension, SummaryTable table, SummaryTableInfo tableInfo, bool bCollated, bool bToCSV, bool bSpreadsheetFriendlyStrings, string weightByColumnNameOverride)
+		void WriteSummaryTableReport(string outputDir, string filenameWithoutExtension, SummaryTable table, List<string> columnFilterList, List<string> rowSortList, bool bCollated, bool bToCSV, bool bSpreadsheetFriendlyStrings, List<SummarySectionBoundaryInfo> sectionBoundaries, string weightByColumnName)
 		{
-			string weightByColumnName = weightByColumnNameOverride != null ? weightByColumnNameOverride : tableInfo.weightByColumn;
-			WriteSummaryTableReport(outputDir, filenameWithoutExtension, table, tableInfo.columnFilterList, tableInfo.rowSortList, bCollated, bToCSV, bSpreadsheetFriendlyStrings, tableInfo.sectionBoundaries, weightByColumnName);
+			SummaryTableInfo tableInfo = new SummaryTableInfo();
+			tableInfo.columnFilterList = columnFilterList;
+			tableInfo.rowSortList = rowSortList;
+			WriteSummaryTableReport(outputDir, filenameWithoutExtension, table, tableInfo, bCollated, bToCSV, bSpreadsheetFriendlyStrings, weightByColumnName);
 		}
 
-		void WriteSummaryTableReport(string outputDir, string filenameWithoutExtension, SummaryTable table, List<string> columnFilterList, List<string> rowSortList, bool bCollated, bool bToCSV, bool bSpreadsheetFriendlyStrings, List<SummarySectionBoundaryInfo> sectionBoundaries, string weightByColumnName)
+		void WriteSummaryTableReport(string outputDir, string filenameWithoutExtension, SummaryTable table, SummaryTableInfo tableInfo, bool bCollated, bool bToCSV, bool bSpreadsheetFriendlyStrings, string weightByColumnNameOverride)
 		{
+			string weightByColumnName = weightByColumnNameOverride != null ? weightByColumnNameOverride : tableInfo.weightByColumn;
 			if (GetBoolArg("noWeightedAvg"))
 			{
 				weightByColumnName = null;
 			}
-			bool reverseSort = GetBoolArg("reverseTable");
-			bool bScrollableTable = GetBoolArg("scrollableTable");
+			bool reverseSort = tableInfo.bReverseSortRows || GetBoolArg("reverseTable");
+			bool bScrollableTable = tableInfo.bScrollableFormatting || GetBoolArg("scrollableTable");
 			bool addMinMaxColumns = !GetBoolArg("noSummaryMinMax");
 			if (!string.IsNullOrEmpty(outputDir))
             {
                 filenameWithoutExtension = Path.Combine(outputDir, filenameWithoutExtension);
             }
-            SummaryTable filteredTable = table.SortAndFilter(columnFilterList, rowSortList, reverseSort, weightByColumnName);
+            SummaryTable filteredTable = table.SortAndFilter(tableInfo.columnFilterList, tableInfo.rowSortList, reverseSort, weightByColumnName);
 			if (bCollated)
 			{
-				filteredTable = filteredTable.CollateSortedTable(rowSortList, addMinMaxColumns);
+				filteredTable = filteredTable.CollateSortedTable(tableInfo.rowSortList, addMinMaxColumns);
 			}
 			if (bToCSV)
 			{
@@ -608,7 +611,7 @@ namespace PerfReportTool
 			{
 				filteredTable.ApplyDisplayNameMapping(statDisplaynameMapping);
 				string VersionString = GetBoolArg("noWatermarks") ? "" : Version.Get();
-				filteredTable.WriteToHTML(filenameWithoutExtension+".html", VersionString, bSpreadsheetFriendlyStrings, sectionBoundaries, bScrollableTable, addMinMaxColumns, GetIntArg("maxSummaryTableStringLength", -1), reportXML.columnFormatInfoList, weightByColumnName);
+				filteredTable.WriteToHTML(filenameWithoutExtension+".html", VersionString, bSpreadsheetFriendlyStrings, tableInfo.sectionBoundaries, bScrollableTable, addMinMaxColumns, GetIntArg("maxSummaryTableStringLength", -1), reportXML.columnFormatInfoList, weightByColumnName);
 			}
 		}
 
diff --git a/Engine/Source/Programs/CSVTools/PerfReportTool/SummaryTable.cs b/Engine/Source/Programs/CSVTools/PerfReportTool/SummaryTable.cs
--- a/Engine/Source/Programs/CSVTools/PerfReportTool/SummaryTable.cs
+++ b/Engine/Source/Programs/CSVTools/PerfReportTool/SummaryTable.cs
@@ -83,6 +83,9 @@ namespace PerfSummaries
 				columnFilterList.AddRange(filterEl.Value.Split(','));
 			}
 
+			bReverseSortRows = tableElement.GetSafeAttibute<bool>("reverseSortRows", false);
+			bScrollableFormatting = tableElement.GetSafeAttibute<bool>("scrollableFormatting", false);
+
 			foreach (XElement sectionBoundaryEl in tableElement.Elements("sectionBoundary"))
 			{
 				if (sectionBoundaryEl != null)
@@ -106,9 +109,15 @@ namespace PerfSummaries
 			rowSortList.AddRange(rowSortStr.Split(','));
 		}
 
+		public SummaryTableInfo()
+		{
+		}
+
 		public List<string> rowSortList = new List<string>();
 		public List<string> columnFilterList = new List<string>();
 		public List<SummarySectionBoundaryInfo> sectionBoundaries = new List<SummarySectionBoundaryInfo>();
+		public bool bReverseSortRows;
+		public bool bScrollableFormatting;
 		public string weightByColumn = null;
 	}
 
