Fix crash when running CsvPerfReport tool with non US culture
- convert floats,doubles with invariant number format
- catch format exception and return default value

#rnx

diff --git a/Engine/Source/Programs/CSVTools/PerfReportTool/PerfReportTool.cs b/Engine/Source/Programs/CSVTools/PerfReportTool/PerfReportTool.cs
--- a/Engine/Source/Programs/CSVTools/PerfReportTool/PerfReportTool.cs
+++ b/Engine/Source/Programs/CSVTools/PerfReportTool/PerfReportTool.cs
@@ -13,6 +13,7 @@ using System.Collections;
 using System.Security.Cryptography;
 
 using PerfSummaries;
+using System.Globalization;
 
 namespace PerfReportTool
 {
@@ -1519,24 +1520,31 @@ namespace PerfReportTool
 		public static T GetSafeAttibute<T>(this XElement element, string attributeName, T defaultValue = default(T))
 		{
 			XAttribute attribute = element.Attribute(attributeName);
-			if (attribute != null)
+			if (attribute == null)
 			{
-				if (typeof(T) == typeof(bool))
-				{
-					return (T)Convert.ChangeType(Convert.ChangeType(attribute.Value, typeof(int)), typeof(bool));
-				}
-				else
+				return defaultValue;
+			}
+
+			try
+			{
+				switch (Type.GetTypeCode(typeof(T)))
 				{
-					return (T)Convert.ChangeType(attribute.Value, typeof(T));
+					case TypeCode.Boolean:
+						return (T)Convert.ChangeType(Convert.ChangeType(attribute.Value, typeof(int)), typeof(bool));
+					case TypeCode.Single:
+					case TypeCode.Double:
+					case TypeCode.Decimal:
+						return (T)Convert.ChangeType(attribute.Value, typeof(T), CultureInfo.InvariantCulture.NumberFormat);
+					default:
+						return (T)Convert.ChangeType(attribute.Value, typeof(T));
 				}
 			}
-			return defaultValue;
+			catch (FormatException e)
+			{
+				Console.WriteLine(string.Format("[Warning] Failed to convert XML attribute '{0}' ({1})", attributeName, e.Message));
+				return defaultValue;
+			}
 		}
-
-
     };
-
-
-
 }
 
