CsvConvert 1.02
-Add input format detection for csvConvert. If not specified, input format and compression will be used for output.
-Add -setMetadata key=value;... arg for overwriting metadata
-(CsvStats) make sure commandline metadata is written at the end when writing out in csv format

diff --git a/Engine/Source/Programs/CSVTools/CsvConvert/CsvConvert.cs b/Engine/Source/Programs/CSVTools/CsvConvert/CsvConvert.cs
--- a/Engine/Source/Programs/CSVTools/CsvConvert/CsvConvert.cs
+++ b/Engine/Source/Programs/CSVTools/CsvConvert/CsvConvert.cs
@@ -13,7 +13,7 @@ namespace CSVInfo
 {
     class Version
     {
-        private static string VersionString = "1.01";
+        private static string VersionString = "1.02";
 
         public static string Get() { return VersionString; }
     };
@@ -62,9 +62,11 @@ namespace CSVInfo
 				"  -outFormat=<csv|bin|csvNoMetadata>\n" +
 				"  [-binCompress=<0|1|2>] (0=none)\n" +
 				"  [-o outFilename]\n" +
-				"  [-force]\n"+
-				"  [-verify]\n"+
-				"  [-dontFixMissingID] - doesn't generate a CsvID if it's not found in the metadata\n";
+				"  [-force] - overwrite even if output file exists already\n" +
+				"  [-inPlace] - overwrite the input file (implies -force)\n" +
+				"  [-verify]\n" +
+				"  [-dontFixMissingID] - doesn't generate a CsvID if it's not found in the metadata\n" +
+				"  [-setMetadata key0=value0;key1=value1;...]\n";
 
 
 			// Read the command line
@@ -74,7 +76,7 @@ namespace CSVInfo
                 return;
             }
 
-
+			bool bInPlace = false;
             ReadCommandLine(args);
 
 			string inFilename;
@@ -99,24 +101,61 @@ namespace CSVInfo
 			}
 			else
 			{
-				// Advanced mode. Output format is required
+				// Advanced mode
 				inFilename = GetArg("in",true);
 				outFormat = GetArg("outformat", true);
-				if (outFormat == "" || inFilename == "")
+				if (inFilename == "")
 				{
-					return;
+					throw new Exception("Need to specify an input filename");
 				}
 				outFilename = GetArg("o", GetArg("out", null));
+				bInPlace = GetBoolArg("inPlace");
+				if (bInPlace)
+				{
+					if (outFilename != null)
+					{
+						throw new Exception("Output filename specified with -inPlace");
+					}
+					outFilename = inFilename;
+				}
+			}
+
+			string inFilenameWithoutExtension;
+			if (isCsvBinFile(inFilename))
+			{
+				inFilenameWithoutExtension = inFilename.Substring(0, inFilename.Length - 8);
+			}
+			else if (isCsvFile(inFilename))
+			{
+				inFilenameWithoutExtension = inFilename.Substring(0, inFilename.Length - 4);
+			}
+			else
+			{
+				throw new Exception("Unexpected input filename extension!");
+			}
+
+			Console.WriteLine("Converting "+inFilename+" to "+outFormat+" format.");
+			Console.WriteLine("Reading input file...");
+
+			bool bFixMissingCsvId = !GetBoolArg("dontFixMissingID");
+			CsvFileInfo inFileInfo = new CsvFileInfo();
+			CsvStats csvStats = CsvStats.ReadCSVFile(inFilename, null, 0, bFixMissingCsvId, inFileInfo);
+
+			// If output format is not set, determine it automatically
+			if (outFormat.Length == 0)
+			{
+				outFormat = inFileInfo.bIsCsvBin ? "bin" : "csv";
 			}
 
+			// Figure out the output format
 			outFormat = outFormat.ToLower();
-			bool bBinOut=false;
+			bool bBinOut = false;
 			bool bWriteMetadata = true;
 			if (outFormat == "bin")
 			{
 				bBinOut = true;
 			}
-			else if (outFormat=="csv")
+			else if (outFormat == "csv")
 			{
 				bBinOut = false;
 			}
@@ -130,45 +169,65 @@ namespace CSVInfo
 				throw new Exception("Unknown output format!");
 			}
 
-			string inFilenameWithoutExtension;
-			if (isCsvBinFile(inFilename))
-			{
-				inFilenameWithoutExtension = inFilename.Substring(0, inFilename.Length - 8);
-			}
-			else if (isCsvFile(inFilename))
-			{
-				inFilenameWithoutExtension = inFilename.Substring(0, inFilename.Length - 4);
-			}
-			else
-			{
-				throw new Exception("Unexpected input filename extension!");
-			}
-			
 			if (outFilename == null)
 			{
 				outFilename = inFilenameWithoutExtension + (bBinOut ? ".csv.bin" : ".csv");
 			}
-
-			if (outFilename.ToLowerInvariant() == inFilename.ToLowerInvariant())
+			if (outFilename.ToLowerInvariant() == inFilename.ToLowerInvariant() && bInPlace == false)
 			{
-				throw new Exception("Input and output filenames can't match!");
+				throw new Exception("Input and output filenames can't match! Specify -inPlace to override");
 			}
-
-			if (System.IO.File.Exists(outFilename) && !GetBoolArg("force"))
+			if (System.IO.File.Exists(outFilename) && !GetBoolArg("force") && !bInPlace)
 			{
 				throw new Exception("Output file already exists! Use -force to overwrite anyway.");
 			}
 
-			Console.WriteLine("Converting "+inFilename+" to "+outFormat+" format. Output filename: "+outFilename);
-			Console.WriteLine("Reading input file...");
+			Console.WriteLine("Output filename: " + outFilename);
+
+
+			// Set metadata if requested
+			string SetMetadataStr = GetArg("setMetadata");
+			if (SetMetadataStr != null && SetMetadataStr.Length > 0)
+			{
+				bool bModified = false;
+				if (csvStats.metaData == null)
+				{
+					throw new Exception("File has no metadata, so can't set it");
+				}
+				string[] KeyValuePairs = SetMetadataStr.Split(';');
+				foreach (string KeyValueStr in KeyValuePairs)
+				{
+					string[] KeyValue = KeyValueStr.Split('=');
+					if (KeyValue.Length == 2)
+					{
+						string Key = KeyValue[0].ToLower().Trim();
+						string Value = KeyValue[1].Trim();
+						Console.WriteLine("Setting metadata: " + Key+"="+Value);
+						csvStats.metaData.Values[Key] = Value;
+						bModified = true;
+					}
+					else
+					{
+						Console.WriteLine("Bad metadata pair: " + KeyValueStr);
+					}
+				}
+				if ( bModified )
+				{
+					csvStats.metaData.Values["metadatamodified"] = "1";
+				}
+			}
 
-			bool bFixMissingCsvId = !GetBoolArg("dontFixMissingID");
-			CsvStats csvStats = CsvStats.ReadCSVFile(inFilename, null, 0, bFixMissingCsvId);
 			Console.WriteLine("Writing output file...");
 			if (bBinOut)
 			{
-				int binCompression = GetIntArg("binCompress", 0);
-				CompressionLevel[] compressionLevels = { CompressionLevel.NoCompression, CompressionLevel.Fastest, CompressionLevel.Optimal };
+				int binCompression = GetIntArg("binCompress", -1);
+
+				// If compression is not set, determine it automatically
+				if (binCompression == -1 )
+				{
+					binCompression = (int)inFileInfo.BinCompressionLevel;
+				}
+
 				if (binCompression < 0 || binCompression > 2)
 				{
 					throw new Exception("Bad compression level specified! Must be 0, 1 or 2");
@@ -180,6 +239,7 @@ namespace CSVInfo
 				csvStats.WriteToCSV(outFilename, bWriteMetadata);
 			}
 
+
 			if (GetBoolArg("verify"))
 			{
 				// TODO: if verify is specified, use a temp intermediate file?
diff --git a/Engine/Source/Programs/CSVTools/CsvStats/CsvStats.cs b/Engine/Source/Programs/CSVTools/CsvStats/CsvStats.cs
--- a/Engine/Source/Programs/CSVTools/CsvStats/CsvStats.cs
+++ b/Engine/Source/Programs/CSVTools/CsvStats/CsvStats.cs
@@ -16,7 +16,7 @@ using System.Security.Cryptography;
 
 namespace CSVStats
 {
-	enum CsvBinVersion
+	public enum CsvBinVersion
 	{
 		PreRelease = 1,
 		InitialRelease,
@@ -43,6 +43,14 @@ namespace CSVStats
 	};
 
 
+	public class CsvFileInfo
+	{
+		public bool bIsCsvBin;
+		public CsvBinVersion BinVersion = CsvBinVersion.COUNT;
+		public CsvBinCompressionLevel BinCompressionLevel = CsvBinCompressionLevel.None;
+	};
+
+
 	public class CsvEvent
     {
 		public CsvEvent(CsvEvent source)
@@ -791,7 +799,7 @@ namespace CSVStats
 		}
 
 
-		public static CsvStats ReadBinFile(string filename, string[] statNamesToRead=null, int numRowsToSkip=0, bool justHeader=false)
+		public static CsvStats ReadBinFile(string filename, string[] statNamesToRead=null, int numRowsToSkip=0, bool justHeader=false, CsvFileInfo FileInfoOut=null)
 		{
 			System.IO.FileStream fileStream = new FileStream(filename, FileMode.Open, FileAccess.Read);
 			System.IO.BinaryReader fileReader = new System.IO.BinaryReader(fileStream);
@@ -857,6 +865,14 @@ namespace CSVStats
 				statNamesToReadDict = csvStatsOut.GetStatNamesMatchingStringList_Dict(statNamesToRead);
 			}
 
+			// Fill in the file info
+			if (FileInfoOut != null)
+			{
+				FileInfoOut.bIsCsvBin = true;
+				FileInfoOut.BinVersion = (CsvBinVersion)version;
+				FileInfoOut.BinCompressionLevel = compressionLevel;
+			}
+
 			if (justHeader)
 			{
 				fileReader.Close();
@@ -1061,7 +1077,10 @@ namespace CSVStats
             if (metaData != null && bWriteMetadataRow )
             {
                 int index = 0;
-                foreach (System.Collections.Generic.KeyValuePair<string, string> pair in metaData.Values)
+				bool bFoundCommandline = false;
+				KeyValuePair<string, string> commandlinePair = new KeyValuePair<string,string>();
+
+				foreach (KeyValuePair<string, string> pair in metaData.Values)
                 {
                     if (index > 0)
                     {
@@ -1076,12 +1095,26 @@ namespace CSVStats
 					}
 					if (pair.Key.ToLower() == "commandline")
 					{
-						value = "\"" + value + "\"";
+						// Defer commandline to last
+						bFoundCommandline = true;
+						commandlinePair = pair;
+						continue;
 					}
 					sb.Append("["+pair.Key + "]," + value);
                     index++;
                 }
-            }
+
+				// Force the commandline to be output last
+				if (bFoundCommandline)
+				{
+					if (index > 0)
+					{
+						sb.Append(",");
+					}
+					string value = "\"" + commandlinePair.Value + "\"";
+					sb.Append("[" + commandlinePair.Key + "]," + value);
+				}
+			}
             csvOutFile.WriteLine(sb);
             csvOutFile.Close();
 
@@ -1500,17 +1533,21 @@ namespace CSVStats
 		}
 
 
-		public static CsvStats ReadCSVFile(string csvFilename, string[] statNames, int numRowsToSkip = 0, bool bGenerateCsvIdIfMissing=false)
+		public static CsvStats ReadCSVFile(string csvFilename, string[] statNames, int numRowsToSkip = 0, bool bGenerateCsvIdIfMissing=false, CsvFileInfo FileInfoOut=null)
         {
 			CsvStats statsOut;
 			if (csvFilename.EndsWith(".csv.bin"))
 			{
-				statsOut = ReadBinFile(csvFilename, statNames, numRowsToSkip);
+				statsOut = ReadBinFile(csvFilename, statNames, numRowsToSkip, false, FileInfoOut);
 			}
 			else
 			{
 				string[] lines = ReadLinesFromFile(csvFilename);
 				statsOut = ReadCSVFromLines(lines, statNames, numRowsToSkip);
+				if ( FileInfoOut != null )
+				{
+					FileInfoOut.bIsCsvBin = false;
+				}
 			}
 			if (bGenerateCsvIdIfMissing)
 			{
