CsvToSVG 2.44 : Make the budget line optional instead of defaulting to 33.33

diff --git a/Engine/Source/Programs/CSVTools/CSVToSVG/CsvToSVG.cs b/Engine/Source/Programs/CSVTools/CSVToSVG/CsvToSVG.cs
--- a/Engine/Source/Programs/CSVTools/CSVToSVG/CsvToSVG.cs
+++ b/Engine/Source/Programs/CSVTools/CSVToSVG/CsvToSVG.cs
@@ -24,8 +24,8 @@ using CSVStats;
 namespace CSVTools
 {
     class Version
-    {
-        private static string VersionString = "2.42";
+    { 
+        private static string VersionString = "2.44";
         
         public static string Get() { return VersionString; }
     };
@@ -184,7 +184,7 @@ namespace CSVTools
         System.IO.StreamWriter svgFile;
         Rect dimensions = new Rect(0, 0, 1000, 500);
         List<CsvStats> csvStats;
-        float budget = 33.33f;
+        float budget = float.MaxValue;
         Theme theme;
         float threshold = -float.MaxValue;
         float averageThreshold= -float.MaxValue;
@@ -577,7 +577,7 @@ namespace CSVTools
             range.MaxY = GetFloatArg("maxy", Range.Auto);
             if (range.MaxX <= 0.0f) range.MaxX = Range.Auto;
             if (range.MaxY <= 0.0f) range.MaxY = Range.Auto;
-            budget = GetFloatArg("budget", 33.333f);
+            budget = GetFloatArg("budget", float.MaxValue);
             float thicknessFactor = GetFloatArg("thickness", 1.0f);
 
             threshold = GetFloatArg("threshold", -float.MaxValue);
@@ -1518,18 +1518,22 @@ namespace CSVTools
             DrawVerticalLine(0, theme.AxisLineColour, rect, range, 0.5f );
             DrawHorizLine(0.0f, theme.AxisLineColour, rect, range, false, 0.5f);
 
-            float budgetLineThickness = theme.BudgetLineThickness;
-            Colour budgetLineColour = new Colour(theme.BudgetLineColour);
-            bool dropShadow = false;
-            if (alpha != 1.0f)
-            {
-                dropShadow = true;
-                budgetLineThickness *= 2.0f;
-                budgetLineColour.r = (byte)((float)budgetLineColour.r * 0.9f);
-                budgetLineColour.g = (byte)((float)budgetLineColour.g * 0.9f);
-                budgetLineColour.b = (byte)((float)budgetLineColour.b * 0.9f);
-            }
-            DrawHorizLine(budget, budgetLineColour, rect, range, true, budgetLineThickness, dropShadow);
+			// Draw the budget line if there is one
+			if (budget != float.MaxValue)
+			{
+				float budgetLineThickness = theme.BudgetLineThickness;
+				Colour budgetLineColour = new Colour(theme.BudgetLineColour);
+				bool dropShadow = false;
+				if (alpha != 1.0f)
+				{
+					dropShadow = true;
+					budgetLineThickness *= 2.0f;
+					budgetLineColour.r = (byte)((float)budgetLineColour.r * 0.9f);
+					budgetLineColour.g = (byte)((float)budgetLineColour.g * 0.9f);
+					budgetLineColour.b = (byte)((float)budgetLineColour.b * 0.9f);
+				}
+				DrawHorizLine(budget, budgetLineColour, rect, range, true, budgetLineThickness, dropShadow);
+			}
         }
 
 		void DrawText(string text, float x, float y, float size, Rect rect, Colour colour, string anchor = "start", string font="Helvetica", string id="", bool dropShadow = false)
