PerfReportTool 4.68 - Fix peak summary not emitting SummaryTableData when detailed reports are disabled

diff --git a/Engine/Source/Programs/CSVTools/PerfReportTool/PerfReportTool.cs b/Engine/Source/Programs/CSVTools/PerfReportTool/PerfReportTool.cs
--- a/Engine/Source/Programs/CSVTools/PerfReportTool/PerfReportTool.cs
+++ b/Engine/Source/Programs/CSVTools/PerfReportTool/PerfReportTool.cs
@@ -19,7 +19,7 @@ namespace PerfReportTool
 {
     class Version
     {
-        private static string VersionString = "4.67";
+        private static string VersionString = "4.68";
 
         public static string Get() { return VersionString; }
     };
diff --git a/Engine/Source/Programs/CSVTools/PerfReportTool/Summaries/PeakSummary.cs b/Engine/Source/Programs/CSVTools/PerfReportTool/Summaries/PeakSummary.cs
--- a/Engine/Source/Programs/CSVTools/PerfReportTool/Summaries/PeakSummary.cs
+++ b/Engine/Source/Programs/CSVTools/PerfReportTool/Summaries/PeakSummary.cs
@@ -86,12 +86,15 @@ namespace PerfSummaries
 
 		void WriteStatSection(StreamWriter htmlFile, CsvStats csvStats, PeakSummarySection section, StreamWriter LLMCsvData, SummaryTableRowData summaryTableRowData)
 		{
-			// Here we are deciding which title we have and write it to the file.
-			htmlFile.WriteLine("<h3>" + section.title + "</h3>");
-			htmlFile.WriteLine("  <table border='0' style='width:400'>");
+			if (htmlFile != null)
+			{
+				// Here we are deciding which title we have and write it to the file.
+				htmlFile.WriteLine("<h3>" + section.title + "</h3>");
+				htmlFile.WriteLine("  <table border='0' style='width:400'>");
 
-			//Hard-coded start of the table.
-			htmlFile.WriteLine("    <tr><td style='width:200'></td><td style='width:75'><b>Average</b></td><td style='width:75'><b>Peak</b></td><td style='width:75'><b>Budget</b></td></tr>");
+				//Hard-coded start of the table.
+				htmlFile.WriteLine("    <tr><td style='width:200'></td><td style='width:75'><b>Average</b></td><td style='width:75'><b>Peak</b></td><td style='width:75'><b>Budget</b></td></tr>");
+			}
 
 			foreach (PeakStatInfo statInfo in section.stats)
 			{
@@ -123,7 +126,10 @@ namespace PerfSummaries
 					averageColour = colorThresholdList.GetColourForValue(average);
 					budgetString = budget.ToString("0");
 				}
-				htmlFile.WriteLine("    <tr><td>" + statInfo.shortName + "</td><td bgcolor=" + averageColour + ">" + average.ToString("0") + "</td><td bgcolor=" + peakColour + ">" + peak.ToString("0") + "</td><td>" + budgetString + "</td></tr>");
+				if (htmlFile != null)
+				{
+					htmlFile.WriteLine("    <tr><td>" + statInfo.shortName + "</td><td bgcolor=" + averageColour + ">" + average.ToString("0") + "</td><td bgcolor=" + peakColour + ">" + peak.ToString("0") + "</td><td>" + budgetString + "</td></tr>");
+				}
 
 				// Pass through color data as part of database-friendly stuff.
 				if (LLMCsvData != null)
@@ -149,7 +155,10 @@ namespace PerfSummaries
 				}
 
 			}
-			htmlFile.WriteLine("  </table>");
+			if (htmlFile != null)
+			{
+				htmlFile.WriteLine("  </table>");
+			}
 		}
 
 		PeakSummarySection FindStatSection(CsvStats csvStats, string statName)
@@ -167,11 +176,6 @@ namespace PerfSummaries
 		public override void WriteSummaryData(System.IO.StreamWriter htmlFile, CsvStats csvStats, bool bWriteSummaryCsv, SummaryTableRowData summaryTableRowData, string htmlFileName)
 		{
 			// Only HTML reporting is supported (does not output summary table row data)
-			if (htmlFile == null)
-			{
-				return;
-			}
-
 			StreamWriter LLMCsvData = null;
 			if (bWriteSummaryCsv)
 			{
@@ -192,7 +196,10 @@ namespace PerfSummaries
 				}
 			}
 
-			htmlFile.WriteLine("<h2>Peaks Summary</h2>");
+			if (htmlFile != null)
+			{
+				htmlFile.WriteLine("<h2>Peaks Summary</h2>");
+			}
 			foreach (PeakSummarySection section in peakSummarySections)
 			{
 				WriteStatSection(htmlFile, csvStats, section, LLMCsvData, summaryTableRowData);
