Expand support for wildcard to support wildcard at start, end or both of filter string.
Fix error message on missing metadata.

diff --git a/Engine/Source/Programs/CSVTools/PerfReportTool/SummaryTable.cs b/Engine/Source/Programs/CSVTools/PerfReportTool/SummaryTable.cs
--- a/Engine/Source/Programs/CSVTools/PerfReportTool/SummaryTable.cs
+++ b/Engine/Source/Programs/CSVTools/PerfReportTool/SummaryTable.cs
@@ -492,7 +492,7 @@ namespace PerfSummaries
 			}
 			if (collateByColumns.Count == 0)
 			{
-				throw new Exception("None of the metadata strings were found:" + collateByList.ToString());
+				throw new Exception("None of the metadata strings were found:" + string.Join(", ", collateByList));
 			}
 
 			// Add the new collateBy columns in the order they appear in the original column list
@@ -689,15 +689,24 @@ namespace PerfSummaries
 			foreach (string filterStr in columnFilterList)
 			{
 				string filterStrLower = filterStr.Trim().ToLower();
-				// Find all matching
-				if (filterStrLower.EndsWith("*"))
+				bool startWild = filterStrLower.StartsWith("*");
+				bool endWild = filterStrLower.EndsWith("*");
+				filterStrLower = filterStrLower.Trim('*');
+				if (startWild && endWild)
+                {
+					orderedKeysWithDupes.AddRange(allMetadataKeys.Where(x => x.Contains(filterStrLower)));
+                }
+				else if(startWild)
+				{
+					orderedKeysWithDupes.AddRange(allMetadataKeys.Where(x => x.EndsWith(filterStrLower)));
+				}
+				else if(endWild)
 				{
-					string prefix = filterStrLower.Trim('*');
 					// Linear search through the sorted key list
 					bool bFound = false;
 					for (int wildcardSearchIndex = 0; wildcardSearchIndex < allMetadataKeys.Count; wildcardSearchIndex++)
 					{
-						if (allMetadataKeys[wildcardSearchIndex].StartsWith(prefix))
+						if (allMetadataKeys[wildcardSearchIndex].StartsWith(filterStrLower))
 						{
 							orderedKeysWithDupes.Add(allMetadataKeys[wildcardSearchIndex]);
 							bFound = true;
