PerfReportTool 4.69: Add -maxFileAgeDays <N> arg - CSV or PRC files older than this will be ignored. This is faster than querying timestamp metadata, especially when reading from network drives.

diff --git a/Engine/Source/Programs/CSVTools/PerfReportTool/PerfReportTool.cs b/Engine/Source/Programs/CSVTools/PerfReportTool/PerfReportTool.cs
--- a/Engine/Source/Programs/CSVTools/PerfReportTool/PerfReportTool.cs
+++ b/Engine/Source/Programs/CSVTools/PerfReportTool/PerfReportTool.cs
@@ -19,7 +19,7 @@ namespace PerfReportTool
 {
     class Version
     {
-        private static string VersionString = "4.68";
+        private static string VersionString = "4.69";
 
         public static string Get() { return VersionString; }
     };
@@ -138,11 +138,12 @@ namespace PerfReportTool
 			"       -allowDuplicateCSVs : doesn't remove duplicate CSVs (Note: can cause summary table cache file locking issues)\n"+
 			"       -requireMetadata : ignores CSVs without metadata\n" +
 			"       -listFiles : just list all files that pass the metadata query. Don't generate any reports.\n" +
-			"       -reportLinkRootPath : Make report links relative to this\n" +
-			"       -csvLinkRootPath : Make CSV file links relative to this\n" +
+			"       -reportLinkRootPath <path> : Make report links relative to this\n" +
+			"       -csvLinkRootPath <path> : Make CSV file links relative to this\n" +
 			"       -weightByColumn : weight collated table averages by this column (overrides value specified in the report XML)\n" +
 			"       -noWeightedAvg : Don't use weighted averages for the collated table\n" +
-			"       -minFrameCount : ignore CSVs without at least this number of valid frames\n" +
+			"       -minFrameCount <n> : ignore CSVs without at least this number of valid frames\n" +
+			"       -maxFileAgeDays <n> : max file age in days. CSV or PRC files older than this will be ignored\n" +
 			"\n" +
 			"Performance args for bulk mode:\n" +
 			"       -precacheCount <n> : number of CSV files to precache in the lookahead cache (0 for no precache)\n" +
@@ -218,6 +219,7 @@ namespace PerfReportTool
 			else
 			{
 				csvDir=GetArg("csvDir");
+				int maxFileAgeDays = GetIntArg("maxFileAgeDays", -1);
 				string summaryTableCacheInDir = GetArg("summaryTableCacheIn");
 				if (csvDir.Length > 0)
 				{
@@ -232,7 +234,7 @@ namespace PerfReportTool
 						searchPattern += ".csv;*.csv.bin";
 					}
 
-					System.IO.FileInfo[] files = GetFilesWithSearchPattern(csvDir, searchPattern, recurse);
+					System.IO.FileInfo[] files = GetFilesWithSearchPattern(csvDir, searchPattern, recurse, maxFileAgeDays);
 					csvFilenames = new string[files.Length];
 					int i = 0;
 					foreach (FileInfo csvFile in files)
@@ -247,7 +249,7 @@ namespace PerfReportTool
 				else if (summaryTableCacheInDir.Length>0)
 				{
 					bool recurse = GetBoolArg("recurse");
-					System.IO.FileInfo[] files = GetFilesWithSearchPattern(summaryTableCacheInDir, "*.prc", recurse);
+					System.IO.FileInfo[] files = GetFilesWithSearchPattern(summaryTableCacheInDir, "*.prc", recurse, maxFileAgeDays);
 					csvFilenames = new string[files.Length];
 					int i = 0;
 					foreach (FileInfo csvFile in files)
@@ -1522,7 +1524,7 @@ namespace PerfReportTool
 			return remStr.Length == 0;
 		}
 
-		System.IO.FileInfo[] GetFilesWithSearchPattern(string directory, string searchPatternStr, bool recurse)
+		System.IO.FileInfo[] GetFilesWithSearchPattern(string directory, string searchPatternStr, bool recurse, int maxFileAgeDays=-1)
 		{
 			List<System.IO.FileInfo> fileList = new List<FileInfo>();
 			string[] searchPatterns = searchPatternStr.Split(';');
@@ -1532,6 +1534,17 @@ namespace PerfReportTool
 				System.IO.FileInfo[] files = di.GetFiles("*.*", recurse ? SearchOption.AllDirectories : SearchOption.TopDirectoryOnly);
 				foreach (FileInfo file in files)
 				{
+					if (maxFileAgeDays >=0)
+					{
+						DateTime fileModifiedTime = file.LastWriteTimeUtc;
+						DateTime currentTime = DateTime.UtcNow;
+						TimeSpan elapsed = currentTime.Subtract(fileModifiedTime);
+						if ( elapsed.TotalHours > (double)maxFileAgeDays*24.0 )
+						{
+							continue;
+						}
+					}
+
 					if (matchesPattern(file.FullName, searchPattern))
 					{
 						fileList.Add(file);
