PerfReportTool 4.70. Make adding duplicate metadata keys non-fatal. A warning will be emitted to the log instead.

diff --git a/Engine/Source/Programs/CSVTools/PerfReportTool/PerfReportTool.cs b/Engine/Source/Programs/CSVTools/PerfReportTool/PerfReportTool.cs
--- a/Engine/Source/Programs/CSVTools/PerfReportTool/PerfReportTool.cs
+++ b/Engine/Source/Programs/CSVTools/PerfReportTool/PerfReportTool.cs
@@ -19,7 +19,7 @@ namespace PerfReportTool
 {
     class Version
     {
-        private static string VersionString = "4.69";
+        private static string VersionString = "4.70";
 
         public static string Get() { return VersionString; }
     };
diff --git a/Engine/Source/Programs/CSVTools/PerfReportTool/SummaryTableData.cs b/Engine/Source/Programs/CSVTools/PerfReportTool/SummaryTableData.cs
--- a/Engine/Source/Programs/CSVTools/PerfReportTool/SummaryTableData.cs
+++ b/Engine/Source/Programs/CSVTools/PerfReportTool/SummaryTableData.cs
@@ -268,7 +268,8 @@ namespace PerfSummaries
 			}
 			catch (System.ArgumentException)
 			{
-				throw new Exception("Summary metadata key " + key + " has already been added");
+				//throw new Exception("Summary metadata key " + key + " has already been added");
+				Console.WriteLine("Warning: Key " + key + " has already been added. Ignoring...");
 			}
 		}
 
@@ -298,7 +299,8 @@ namespace PerfSummaries
 			}
 			catch (System.ArgumentException)
 			{
-				throw new Exception("Summary metadata key " + key + " has already been added");
+				//throw new Exception("Summary metadata key " + key + " has already been added");
+				Console.WriteLine("Warning: Key " + key + " has already been added. Ignoring...");
 			}
 		}
 
