Fix bug with stat wildcard matching which caused stats to report a match if a wildcard was present even if the start or end of the stat name didn't match

diff --git a/Engine/Source/Programs/CSVTools/CsvStats/CsvStats.cs b/Engine/Source/Programs/CSVTools/CsvStats/CsvStats.cs
--- a/Engine/Source/Programs/CSVTools/CsvStats/CsvStats.cs
+++ b/Engine/Source/Programs/CSVTools/CsvStats/CsvStats.cs
@@ -581,6 +581,21 @@ namespace CSVStats
 
 		private static bool WildcardSubstringMatch(string strLower, string[] subSearchStrings)
 		{
+			if ( subSearchStrings.Length == 0 )
+			{
+				return false;
+			}
+			// Make sure the first and last substrings match the start and end
+			if (!strLower.StartsWith(subSearchStrings[0]))
+			{
+				return false;
+			}
+			if (!strLower.EndsWith(subSearchStrings.Last()))
+			{
+				return false;
+			}
+
+			// Make sure we match all the substrings, in order
 			string remainingString = strLower;
 			foreach (string subString in subSearchStrings)
 			{
