CsvInfo 1.02 - json support. Pass -toJson <filename> to output Csv metadata and stats to a json file.

diff --git a/Engine/Source/Programs/CSVTools/CsvInfo/CSVInfo.cs b/Engine/Source/Programs/CSVTools/CsvInfo/CSVInfo.cs
--- a/Engine/Source/Programs/CSVTools/CsvInfo/CSVInfo.cs
+++ b/Engine/Source/Programs/CSVTools/CsvInfo/CSVInfo.cs
@@ -12,7 +12,7 @@ namespace CSVInfo
 {
     class Version
     {
-        private static string VersionString = "1.01";
+        private static string VersionString = "1.02";
 
         public static string Get() { return VersionString; }
     };
@@ -42,16 +42,43 @@ namespace CSVInfo
             return 0;
         }
 
+		string Quotify(string s)
+		{
+			return "\"" + s + "\"";
+		}
+
+		string Sanitize(string s)
+		{
+			return s.Replace("\\","\\\\").Replace("\"", "\\\"");
+		}
+
+		string ToJsonString(string s)
+		{
+			return Quotify(Sanitize(s));
+		}
+
+		string ToJsonStringList(List<string> list)
+		{
+			List<string> safeList = new List<string>();
+			foreach(string s in list)
+			{
+				safeList.Add(ToJsonString(s));
+			}
+			return "[" + String.Join(",", safeList) + "]";
+		}
 
         void Run(string[] args)
         {
             string formatString =
                 "Format: \n" +
-                "<csvfilename>";
+                "  <csvfilename>\n"+
+				"  [-showaverages]\n"+
+				"  [-toJson <filename>]";
 
-            // Read the command line
-            if (args.Length < 1)
+			// Read the command line
+			if (args.Length < 1)
             {
+				WriteLine("CsvInfo " + Version.Get());
                 WriteLine(formatString);
                 return;
             }
@@ -61,36 +88,75 @@ namespace CSVInfo
             ReadCommandLine(args);
 
             bool showAverages = GetBoolArg("showAverages");
+			string jsonFilename = GetArg("toJson",false);
 
-            // Write out the sorted stat names
-            Console.Out.WriteLine("Stats:");
-            CsvStats csvStats = CsvStats.ReadCSVFile(csvFilename, null);
-            List<string> statLines = new List<string>();
-            foreach (StatSamples stat in csvStats.Stats.Values.ToArray())
-            {
-                string statLine = stat.Name;
-                if (showAverages)
-                {
-                    statLine += " (" + stat.average.ToString() + ")";
-                }
-                statLines.Add(statLine);
-            }
-            statLines.Sort();
-            foreach (string statLine in statLines)
-            {
-                Console.Out.WriteLine("  " + statLine);
-            }
+			CsvStats csvStats = CsvStats.ReadCSVFile(csvFilename, null);
+			List<string> statLines = new List<string>();
+			foreach (StatSamples stat in csvStats.Stats.Values.ToArray())
+			{
+				string statLine = stat.Name;
+				if (showAverages)
+				{
+					statLine += " (" + stat.average.ToString() + ")";
+				}
+				statLines.Add(statLine);
+			}
+			statLines.Sort();
 
-            if ( csvStats.metaData != null )
-            {
-                // Write out the metadata, if it exists
-                Console.Out.WriteLine("\nMetadata:");
-                foreach (KeyValuePair<string, string> pair in csvStats.metaData.Values.ToArray())
-                {
-                    string key = pair.Key.PadRight(20);
-                    Console.Out.WriteLine("  " + key + ": " + pair.Value);
-                }
-            }
-        }
+
+			if (jsonFilename != "")
+			{
+				// We just write the lines raw, since this version of .Net doesn't have a json serializer. 
+				// TODO: Fix this when we upgrade to .Net 5.0 and use System.Text.Json
+				List<string> jsonLines = new List<string>();
+				jsonLines.Add("{");
+				if (csvStats.metaData != null)
+				{
+					jsonLines.Add("  \"metadata\":\n  {");
+					Dictionary<string, string> metadata = csvStats.metaData.Values;
+					int count = metadata.Count;
+					int index = 0;
+					foreach (string key in metadata.Keys)
+					{
+						string line = "    " + ToJsonString(key) + ":" + ToJsonString(metadata[key]);
+						if (index < count-1)
+						{
+							line += ",";
+						}
+						jsonLines.Add(line);
+						index++;
+					}
+					jsonLines.Add("  },");
+				}
+				jsonLines.Add("  \"stats\":\n  {");
+				jsonLines.Add("    "+ToJsonStringList(statLines));
+				jsonLines.Add("  }");
+
+				jsonLines.Add("}");
+				System.IO.File.WriteAllLines(jsonFilename,jsonLines);
+				Console.Out.WriteLine("Wrote csv info to " + jsonFilename);
+			}
+			else
+			{
+				// Write out the sorted stat names
+				Console.Out.WriteLine("Stats:");
+				foreach (string statLine in statLines)
+				{
+					Console.Out.WriteLine("  " + statLine);
+				}
+
+				if (csvStats.metaData != null)
+				{
+					// Write out the metadata, if it exists
+					Console.Out.WriteLine("\nMetadata:");
+					foreach (KeyValuePair<string, string> pair in csvStats.metaData.Values.ToArray())
+					{
+						string key = pair.Key.PadRight(20);
+						Console.Out.WriteLine("  " + key + ": " + pair.Value);
+					}
+				}
+
+			}
+		}
     }
 }
