CsvInfo - Fix stat array json output

diff --git a/Engine/Source/Programs/CSVTools/CsvInfo/CSVInfo.cs b/Engine/Source/Programs/CSVTools/CsvInfo/CSVInfo.cs
--- a/Engine/Source/Programs/CSVTools/CsvInfo/CSVInfo.cs
+++ b/Engine/Source/Programs/CSVTools/CsvInfo/CSVInfo.cs
@@ -112,7 +112,7 @@ namespace CSVInfo
 				jsonLines.Add("{");
 				if (csvStats.metaData != null)
 				{
-					jsonLines.Add("  \"metadata\":\n  {");
+					jsonLines.Add("  \"metadata\": {");
 					Dictionary<string, string> metadata = csvStats.metaData.Values;
 					int count = metadata.Count;
 					int index = 0;
@@ -128,9 +128,7 @@ namespace CSVInfo
 					}
 					jsonLines.Add("  },");
 				}
-				jsonLines.Add("  \"stats\":\n  {");
-				jsonLines.Add("    "+ToJsonStringList(statLines));
-				jsonLines.Add("  }");
+				jsonLines.Add("  \"stats\": "+ ToJsonStringList(statLines));
 
 				jsonLines.Add("}");
 				System.IO.File.WriteAllLines(jsonFilename,jsonLines);
