PerfReportTool updates to support new PerfReportServer

PerfReportTool 4.71
 - Add support for summary table link templates that can be replaced later (e.g by a web server). E.g {LinkTemplate:Csv:<CSVID>}. Pass -linkTemplates to enable. This works for Csv and Report links, as well as ExtraLinks
 - Add -csvList and -PRCList args for specifying a list of PRCs or CSVs
 - Add -summaryTableXML arg for specifying an additional XML file containing summary tables
 - Summary table titles can now be specified with -summaryTitle instead of this being automatic

All CSV tools:
 - Add response file commandline support via -response=<filename>

diff --git a/Engine/Source/Programs/CSVTools/CsvStats/CommandLineTool.cs b/Engine/Source/Programs/CSVTools/CsvStats/CommandLineTool.cs
--- a/Engine/Source/Programs/CSVTools/CsvStats/CommandLineTool.cs
+++ b/Engine/Source/Programs/CSVTools/CsvStats/CommandLineTool.cs
@@ -22,8 +22,64 @@ namespace CSVStats
 			return commandLine;
 		}
 
+		private List<string> GetArgList(string inCommandLine)
+		{
+			List<string> args = new List<string>();
+			bool bInQuotes = false;
+			string currentArg = "";
+			for (int i=0; i< inCommandLine.Length; i++)
+			{
+				char c = inCommandLine[i];
+				if (c=='"')
+				{
+					bInQuotes = !bInQuotes;
+				}
+				else
+				{
+					if (bInQuotes)
+					{
+						currentArg += c;
+					}
+					else
+					{ 
+						if (Char.IsWhiteSpace(c))
+						{
+							if (currentArg.Length > 0)
+							{
+								args.Add(currentArg);
+								currentArg = "";
+							}
+						}
+						else
+						{
+							currentArg += c;
+						}
+
+					}
+				}
+			}
+			if (currentArg.Length > 0)
+			{
+				args.Add(currentArg);
+			}
+			return args;
+		}
+
 		public CommandLine(string[] args)
 		{
+			// Read commandline as a response file if specified (must be the only two args)
+			if (args.Length==2 && args[0].ToLower()=="-response")
+			{
+				string [] lines = File.ReadAllLines(args[1]);
+				List<string> argList = new List<string>();
+				foreach (string line in lines)
+				{
+					argList.AddRange(GetArgList(line));
+				}
+				args = argList.ToArray();
+			}
+
+			// Write the commandline as a string for reporting purposes
 			commandLine = "";
 			foreach (string arg in args)
 			{
@@ -36,6 +92,8 @@ namespace CSVStats
 					commandLine += arg + " ";
 				}
 			}
+
+			// Parse the commandline
 			CommandLineArgs = new Dictionary<string, string>();
 			for (int i = 0; i < args.Length; i++)
 			{
diff --git a/Engine/Source/Programs/CSVTools/PerfReportTool/PerfReportTool.cs b/Engine/Source/Programs/CSVTools/PerfReportTool/PerfReportTool.cs
--- a/Engine/Source/Programs/CSVTools/PerfReportTool/PerfReportTool.cs
+++ b/Engine/Source/Programs/CSVTools/PerfReportTool/PerfReportTool.cs
@@ -19,7 +19,7 @@ namespace PerfReportTool
 {
     class Version
     {
-        private static string VersionString = "4.70";
+        private static string VersionString = "4.71";
 
         public static string Get() { return VersionString; }
     };
@@ -76,7 +76,8 @@ namespace PerfReportTool
 			"PerfReportTool v" + Version.Get() + "\n" +
 			"\n" +
 			"Format: \n" +
-			"       -csv <filename> or -csvdir <directory path> or -summaryTableCacheIn <directory path>\n" +
+			"       -csv <filename> or -csvdir <directory path> or -summaryTableCacheIn <directory path> or\n" +
+			"       -csvList <comma separated> or -prcList <comma separated>\n" +
 			"       -o <dir name>: output directory (will be created if necessary)\n" +
 			"\n" +
 			"Optional Args:\n" +
@@ -85,7 +86,8 @@ namespace PerfReportTool
 			"       -graphXML <xmlfilename>\n" +
 			"       -reportXML <xmlfilename>\n" +
 			"       -reportxmlbasedir <folder>\n" +
-			"       -title <name>\n" +
+			"       -title <name> - title for detailed reports\n" +
+			"       -summaryTitle <name> - title for summary tables\n" +
 			"       -maxy <value> - forces all graphs to use this value\n" +
 			"       -writeSummaryCsv : if specified, a csv file containing summary information will be generated.\n"+
 			"          Not available in bulk mode.\n" +
@@ -109,15 +111,16 @@ namespace PerfReportTool
 			"       -endEvent <event> : strip data after this event\n" +
 			"       -noStripEvents : if specified, don't strip out samples between excluded events from the stats\n" +
 			"\n" +
-			"Optional bulk mode args: (use with -csvdir or -summaryTableCacheIn)\n" +
+			"Optional bulk mode args: (use with -csvdir, -summaryTableCacheIn, -csvList, -prcList)\n" +
 			"       -recurse \n" +
 			"       -searchpattern <pattern>, e.g -searchpattern csvprofile*\n" +
-			"       -customTable <comma seprated fields>\n" +
-			"       -customTableSort <comma seprated field row sort order> (use with -customTable)\n" +
+			"       -customTable <comma separated fields>\n" +
+			"       -customTableSort <comma separated field row sort order> (use with -customTable)\n" +
 			"       -noDetailedReports : skips individual report generation\n" +
 			"       -collateTable : writes a collated table in addition to the main one, merging by row sort\n" +
 			"       -emailTable : writes a condensed email-friendly table (see the 'condensed' summary table)\n" +
 			"       -csvTable : writes the summary table in CSV format instead of html\n" +
+			"       -summaryTableXML <XML filename>\n" +
 			"       -summaryTable <name> :\n" +
 			"           Selects a custom summary table type from the list in reportTypes.xml \n"+
 			"           (if not specified, 'default' will be used)\n" +
@@ -140,6 +143,7 @@ namespace PerfReportTool
 			"       -listFiles : just list all files that pass the metadata query. Don't generate any reports.\n" +
 			"       -reportLinkRootPath <path> : Make report links relative to this\n" +
 			"       -csvLinkRootPath <path> : Make CSV file links relative to this\n" +
+			"       -linkTemplates : insert templates instead of links that can be replaced later, e.g {{LinkTemplate:Report:<CSV ID>}}\n" +
 			"       -weightByColumn : weight collated table averages by this column (overrides value specified in the report XML)\n" +
 			"       -noWeightedAvg : Don't use weighted averages for the collated table\n" +
 			"       -minFrameCount <n> : ignore CSVs without at least this number of valid frames\n" +
@@ -221,6 +225,8 @@ namespace PerfReportTool
 				csvDir=GetArg("csvDir");
 				int maxFileAgeDays = GetIntArg("maxFileAgeDays", -1);
 				string summaryTableCacheInDir = GetArg("summaryTableCacheIn");
+				string csvListStr = GetArg("csvList");
+				string prcListStr = GetArg("prcList");
 				if (csvDir.Length > 0)
 				{
 					bool recurse = GetBoolArg("recurse");
@@ -244,9 +250,9 @@ namespace PerfReportTool
 					}
 					// We don't write summary CSVs in bulk mode
 					bBulkMode = true;
-                    perfLog.LogTiming("DirectoryScan");
-                }
-				else if (summaryTableCacheInDir.Length>0)
+					perfLog.LogTiming("DirectoryScan");
+				}
+				else if (summaryTableCacheInDir.Length > 0)
 				{
 					bool recurse = GetBoolArg("recurse");
 					System.IO.FileInfo[] files = GetFilesWithSearchPattern(summaryTableCacheInDir, "*.prc", recurse, maxFileAgeDays);
@@ -261,16 +267,30 @@ namespace PerfReportTool
 					bSummaryTableCacheOnlyMode = true;
 					perfLog.LogTiming("DirectoryScan");
 				}
-				else
+				else if (csvListStr.Length > 0)
+				{
+					csvFilenames = csvListStr.Split(',');
+					bBulkMode = true;
+				}
+				else if (prcListStr.Length > 0)
 				{
+					csvFilenames = prcListStr.Split(',');
+					bBulkMode = true;
+					bSummaryTableCacheOnlyMode = true;
+				}
+				else
+				{ 
 					string csvFilenamesStr = GetArg("csv");
 					if (csvFilenamesStr.Length == 0)
 					{
 						csvFilenamesStr = GetArg("csvs", true);
 						if (csvFilenamesStr.Length == 0)
 						{
-							System.Console.Write(formatString);
-							return;
+							if (!GetBoolArg("listSummaryTables"))
+							{
+								System.Console.Write(formatString);
+								return;
+							}
 						}
 					}
 					csvFilenames = csvFilenamesStr.Split(';');
@@ -278,7 +298,7 @@ namespace PerfReportTool
 			}
 
 			// Load the report + graph XML data
-			reportXML = new ReportXML(GetArg("graphxml", false), GetArg("reportxml", false), GetArg("reportxmlbasedir", false));
+			reportXML = new ReportXML(GetArg("graphxml", false), GetArg("reportxml", false), GetArg("reportxmlbasedir", false), GetArg("summaryTableXml", false));
 
 			if (GetBoolArg("listSummaryTables"))
 			{
@@ -642,7 +662,8 @@ namespace PerfReportTool
 			{
 				filteredTable.ApplyDisplayNameMapping(statDisplaynameMapping);
 				string VersionString = GetBoolArg("noWatermarks") ? "" : Version.Get();
-				filteredTable.WriteToHTML(filenameWithoutExtension+".html", VersionString, bSpreadsheetFriendlyStrings, tableInfo.sectionBoundaries, bScrollableTable, addMinMaxColumns, GetIntArg("maxSummaryTableStringLength", -1), reportXML.columnFormatInfoList, weightByColumnName);
+				string summaryTitle = GetArg("summaryTitle", null);
+				filteredTable.WriteToHTML(filenameWithoutExtension+".html", VersionString, bSpreadsheetFriendlyStrings, tableInfo.sectionBoundaries, bScrollableTable, addMinMaxColumns, GetIntArg("maxSummaryTableStringLength", -1), reportXML.columnFormatInfoList, weightByColumnName, summaryTitle);
 			}
 		}
 
@@ -898,6 +919,13 @@ namespace PerfReportTool
 				Uri relativeCsvUri = finalDirUri.MakeRelativeUri(csvFileUri);
 				string csvPath = relativeCsvUri.ToString();
 
+				bool bLinkTemplates = GetBoolArg("linkTemplates");
+				string csvId = null;
+				if (csvStats.metaData != null)
+				{
+					csvId=csvStats.metaData.GetValue("csvid", null);
+				}
+
 				// re-root the CSV path if requested
 				string csvLinkRootPath = GetArg("csvLinkRootPath", null);
 				if ( csvDir != null && csvLinkRootPath != null)
@@ -910,7 +938,13 @@ namespace PerfReportTool
 					csvPath = new Uri(csvPath, UriKind.Absolute).ToString();
 				}
 
-				rowData.Add(SummaryTableElement.Type.ToolMetadata, "Csv File", "<a href='" + csvPath + "'>" + shortName + ".csv" + "</a>", null, csvPath);
+				string csvLink = "<a href='" + csvPath + "'>" + shortName + ".csv" + "</a>";
+				if (bLinkTemplates)
+				{
+					csvLink = "{LinkTemplate:Csv:" + (csvId ?? "0") + "}";
+				}
+
+				rowData.Add(SummaryTableElement.Type.ToolMetadata, "Csv File", csvLink, null, csvPath);
 				rowData.Add(SummaryTableElement.Type.ToolMetadata, "ReportType", reportTypeInfo.name);
 				rowData.Add(SummaryTableElement.Type.ToolMetadata, "ReportTypeID", reportTypeInfo.summaryTableCacheID);
 				if (htmlFilename != null)
@@ -921,8 +955,12 @@ namespace PerfReportTool
 					{
 						htmlUrl = reportLinkRootPath + htmlFilename;
 					}
-
-					rowData.Add(SummaryTableElement.Type.ToolMetadata, "Report", "<a href='" + htmlUrl + "'>Link</a>");
+					string reportLink = "<a href='" + htmlUrl + "'>Link</a>";
+					if (bLinkTemplates)
+					{
+						reportLink = "{LinkTemplate:Report:" + (csvId ?? "0") + "}";
+					}	
+					rowData.Add(SummaryTableElement.Type.ToolMetadata, "Report", reportLink);
 				}
 				// Pass through all the metadata from the CSV
 				if (csvStats.metaData != null)
@@ -1144,7 +1182,8 @@ namespace PerfReportTool
 			bool bExtraLinksSummary = GetBoolArg("extraLinksSummary");
 			if (bExtraLinksSummary)
 			{
-				summaries.Insert(0,new ExtraLinksSummary(null, null));
+				bool bLinkTemplates = GetBoolArg("linkTemplates");
+				summaries.Insert(0,new ExtraLinksSummary(null, null, bLinkTemplates));
 			}
 
 			// If the reporttype has summary info, then write out the summary]
diff --git a/Engine/Source/Programs/CSVTools/PerfReportTool/ReportXml.cs b/Engine/Source/Programs/CSVTools/PerfReportTool/ReportXml.cs
--- a/Engine/Source/Programs/CSVTools/PerfReportTool/ReportXml.cs
+++ b/Engine/Source/Programs/CSVTools/PerfReportTool/ReportXml.cs
@@ -67,7 +67,7 @@ namespace PerfReportTool
 			return false;
 		}
 
-		public ReportXML(string graphXMLFilenameIn, string reportXMLFilenameIn, string baseXMLDirectoryOverride)
+		public ReportXML(string graphXMLFilenameIn, string reportXMLFilenameIn, string baseXMLDirectoryOverride, string additionalSummaryTableXmlFilename)
 		{
 			string location = System.Reflection.Assembly.GetEntryAssembly().Location.ToLower();
 			string baseDirectory = location.Replace("perfreporttool.exe", "");
@@ -83,14 +83,7 @@ namespace PerfReportTool
 			// Check if the base directory is being overridden
 			if (baseXMLDirectoryOverride.Length > 0)
 			{
-				if (IsAbsolutePath(baseXMLDirectoryOverride))
-				{
-					baseDirectory = baseXMLDirectoryOverride;
-				}
-				else
-				{
-					baseDirectory = Path.Combine(baseDirectory, baseXMLDirectoryOverride);
-				}
+				baseDirectory = IsAbsolutePath(baseXMLDirectoryOverride) ? baseXMLDirectoryOverride : Path.Combine(baseDirectory, baseXMLDirectoryOverride);
 			}
 			Console.Out.WriteLine("BaseDir:   " + baseDirectory);
 
@@ -100,15 +93,7 @@ namespace PerfReportTool
 			reportTypeXmlFilename = Path.Combine(baseDirectory, "reportTypes.xml");
 			if (reportXMLFilenameIn.Length > 0)
 			{
-				// Check if this is an absolute path
-				if (IsAbsolutePath(reportXMLFilenameIn))
-				{
-					reportTypeXmlFilename = reportXMLFilenameIn;
-				}
-				else
-				{
-					reportTypeXmlFilename = Path.Combine(baseDirectory, reportXMLFilenameIn);
-				}
+				reportTypeXmlFilename = IsAbsolutePath(reportXMLFilenameIn) ? reportXMLFilenameIn : Path.Combine(baseDirectory, reportXMLFilenameIn);
 			}
 			Console.Out.WriteLine("ReportXML: " + reportTypeXmlFilename);
 			XDocument reportTypesDoc = XDocument.Load(reportTypeXmlFilename);
@@ -118,6 +103,35 @@ namespace PerfReportTool
 				throw new Exception("No root element found in report XML " + reportTypeXmlFilename);
 			}
 
+			// Read the additional summary table XML (if supplied)
+			if (additionalSummaryTableXmlFilename.Length > 0)
+			{
+				if (!IsAbsolutePath(additionalSummaryTableXmlFilename))
+				{
+					additionalSummaryTableXmlFilename = Path.Combine(baseDirectory, additionalSummaryTableXmlFilename);
+				}
+				XDocument summaryXmlDoc = XDocument.Load(additionalSummaryTableXmlFilename);
+				XElement summaryTablesEl = summaryXmlDoc.Element("summaryTables");
+				if (summaryTablesEl == null)
+				{
+					throw new Exception("No summaryTables element found in summaryTableXML file: " + additionalSummaryTableXmlFilename);
+				}
+
+				XElement destSummaryTablesEl = rootElement.Element("summaryTables");
+				if (destSummaryTablesEl == null)
+				{
+					rootElement.Add(summaryTablesEl);
+				}
+				else
+				{
+					foreach (XElement child in summaryTablesEl.Elements())
+					{
+						destSummaryTablesEl.Add(child);
+					}
+				}
+			}
+			Console.Out.WriteLine("ReportXML: " + reportTypeXmlFilename);
+
 			reportTypesElement = rootElement.Element("reporttypes");
 			if (reportTypesElement == null)
 			{
@@ -128,14 +142,7 @@ namespace PerfReportTool
 			string graphsXMLFilename;
 			if (graphXMLFilenameIn.Length > 0)
 			{
-				if (IsAbsolutePath(graphXMLFilenameIn))
-				{
-					graphsXMLFilename = graphXMLFilenameIn;
-				}
-				else
-				{
-					graphsXMLFilename = Path.Combine(baseDirectory, graphXMLFilenameIn);
-				}
+				graphsXMLFilename = IsAbsolutePath(graphXMLFilenameIn) ? graphXMLFilenameIn : Path.Combine(baseDirectory, graphXMLFilenameIn);
 			}
 			else
 			{
@@ -443,12 +450,6 @@ namespace PerfReportTool
 			return summaryTables.Keys.ToList();
 		}
 
-		public string GetReportXmlFilename()
-		{
-			return reportTypeXmlFilename;
-		}
-
-
 		Dictionary<string, SummaryTableInfo> summaryTables;
 
 		XElement reportTypesElement;
diff --git a/Engine/Source/Programs/CSVTools/PerfReportTool/Summaries/ExtraLinksSummary.cs b/Engine/Source/Programs/CSVTools/PerfReportTool/Summaries/ExtraLinksSummary.cs
--- a/Engine/Source/Programs/CSVTools/PerfReportTool/Summaries/ExtraLinksSummary.cs
+++ b/Engine/Source/Programs/CSVTools/PerfReportTool/Summaries/ExtraLinksSummary.cs
@@ -14,35 +14,43 @@ namespace PerfSummaries
 	{
 		class ExtraLink
 		{
-			public ExtraLink(string fileLine)
+			public ExtraLink(string fileLine, string inLinkTemplateCsvId)
 			{
 				string[] Sections = fileLine.Split(',');
 				if (Sections.Length != 3)
-				{
+                {
 					throw new Exception("Bad links line format: " + fileLine);
 				}
-				LongName = Sections[0];
-				ShortName = Sections[1];
+                ReportText = Sections[0];
+                SummaryTableText = Sections[1];
 				LinkURL = Sections[2];
-			}
-			public string GetLinkString(bool bUseLongName)
-			{
-				string Text = bUseLongName ? LongName : ShortName;
-				return "<a href='" + LinkURL + "'>" + Text + "</a>";
-			}
-			public string LongName;
-			public string ShortName;
+                LinkTemplateCsvId = inLinkTemplateCsvId;
+            }
+            public string GetLinkString(bool bIsSummaryTable)
+            {
+                if (bIsSummaryTable && LinkTemplateCsvId != null)
+                {
+                    // Output a link template if link templates are enabled and we have a CSV
+                    return "{LinkTemplate:" + SummaryTableText.Replace(" ", "") + ":" + LinkTemplateCsvId + "}";
+                }
+                string Text = bIsSummaryTable ? SummaryTableText : ReportText;
+                return "<a href='" + LinkURL + "'>" + Text + "</a>";
+            }
+            public string ReportText;
+			public string SummaryTableText;
 			public string LinkURL;
-		};
+            public string LinkTemplateCsvId;
+        };
+        public ExtraLinksSummary(XElement element, string baseXmlDirectory, bool bInLinkTemplates=false)
+        {
+            title = "Links";
+            if (element != null)
+            {
+                title = element.GetSafeAttibute("title", title);
+            }
+            bLinkTemplates = bInLinkTemplates;
+        }
 
-		public ExtraLinksSummary(XElement element, string baseXmlDirectory)
-		{
-			title = "Links";
-			if (element != null)
-			{
-				title = element.GetSafeAttibute("title", title);
-			}
-		}
 		public ExtraLinksSummary() { }
 
 		public override string GetName() { return "extralinks"; }
@@ -51,7 +59,14 @@ namespace PerfSummaries
 		{
 			List<ExtraLink> links = new List<ExtraLink>();
 
-			string csvFilename = csvStats.metaData.GetValue("csvfilename", null);
+            string csvId = null;
+            if (bLinkTemplates && csvStats.metaData != null)
+            {
+                csvId = csvStats.metaData.GetValue("csvid", null);
+            }
+
+
+            string csvFilename = csvStats.metaData.GetValue("csvfilename", null);
 			if (csvFilename == null)
 			{
 				Console.WriteLine("Can't find CSV filename for ExtraLinks summary. Skipping");
@@ -67,7 +82,7 @@ namespace PerfSummaries
 			string[] lines = File.ReadAllLines(linksFilename);
 			foreach (string line in lines)
 			{
-				links.Add(new ExtraLink(line));
+				links.Add(new ExtraLink(line, csvId));
 			}
 			if (links.Count == 0)
 			{
@@ -81,7 +96,7 @@ namespace PerfSummaries
 				htmlFile.WriteLine("  <ul>");
 				foreach (ExtraLink link in links)
 				{
-					htmlFile.WriteLine("  <li>" + link.GetLinkString(true) + "</li>");
+					htmlFile.WriteLine("  <li>" + link.GetLinkString(false) + "</li>");
 				}
 				htmlFile.WriteLine("  </ul>");
 			}
@@ -91,14 +106,15 @@ namespace PerfSummaries
 			{
 				foreach (ExtraLink link in links)
 				{
-					rowData.Add(SummaryTableElement.Type.SummaryTableMetric, link.LongName, link.GetLinkString(false), null);
+					rowData.Add(SummaryTableElement.Type.SummaryTableMetric, link.ReportText, link.GetLinkString(true), null);
 				}
 			}
 		}
 		public override void PostInit(ReportTypeInfo reportTypeInfo, CsvStats csvStats)
 		{
-		}
+        }
 		string title;
+        bool bLinkTemplates;
 	};
 
 }
\ No newline at end of file
diff --git a/Engine/Source/Programs/CSVTools/PerfReportTool/SummaryTable.cs b/Engine/Source/Programs/CSVTools/PerfReportTool/SummaryTable.cs
--- a/Engine/Source/Programs/CSVTools/PerfReportTool/SummaryTable.cs
+++ b/Engine/Source/Programs/CSVTools/PerfReportTool/SummaryTable.cs
@@ -826,7 +826,7 @@ namespace PerfSummaries
 		}
 
 
-		public void WriteToHTML(string htmlFilename, string VersionString, bool bSpreadsheetFriendlyStrings, List<SummarySectionBoundaryInfo> sectionBoundaries, bool bScrollableTable, bool bAddMinMaxColumns, int maxColumnStringLength, SummaryTableColumnFormatInfoCollection columnFormatInfoList, string weightByColumnName)
+		public void WriteToHTML(string htmlFilename, string VersionString, bool bSpreadsheetFriendlyStrings, List<SummarySectionBoundaryInfo> sectionBoundaries, bool bScrollableTable, bool bAddMinMaxColumns, int maxColumnStringLength, SummaryTableColumnFormatInfoCollection columnFormatInfoList, string weightByColumnName, string title)
 		{
 			System.IO.StreamWriter htmlFile = new System.IO.StreamWriter(htmlFilename, false);
 			int statColSpan = hasMinMaxColumns ? 3 : 1;
@@ -836,8 +836,15 @@ namespace PerfSummaries
 				cellPadding = 4;
 			}
 
+			// Generate an automatic title
+			if (title==null)
+			{
+				title = htmlFilename.Replace("_Email.html", "").Replace(".html", "").Replace("\\", "/");
+				title = title.Substring(title.LastIndexOf('/') + 1);
+			}
+
 			htmlFile.WriteLine("<html>");
-			htmlFile.WriteLine("<head><title>Performance summary</title>");
+			htmlFile.WriteLine("<head><title>Perf Summary: "+ title + "</title>");
 
 			// Figure out the sticky column count
 			int stickyColumnCount = 0;
@@ -1011,9 +1018,6 @@ namespace PerfSummaries
 				string TopHeaderRow = "";
 				if (bScrollableTable)
 				{
-					// Generate an automatic title
-					string title = htmlFilename.Replace("_Email.html", "").Replace(".html", "").Replace("\\", "/");
-					title = title.Substring(title.LastIndexOf('/') + 1);
 					TopHeaderRow += "<th colspan='" + firstStatColumnIndex + "'><h3>" + title + "</h3></th>";
 				}
 				else
