CsvConvert 1.03
 - Add setMetadataFile arg for setting/overriding metadata. This is more robust than passing it on the commandline. Input is a text file where each line corresponds to a metadata key/value and each line has the format <Key>=<Value>

diff --git a/Engine/Source/Programs/CSVTools/CsvConvert/CsvConvert.cs b/Engine/Source/Programs/CSVTools/CsvConvert/CsvConvert.cs
--- a/Engine/Source/Programs/CSVTools/CsvConvert/CsvConvert.cs
+++ b/Engine/Source/Programs/CSVTools/CsvConvert/CsvConvert.cs
@@ -13,7 +13,7 @@ namespace CSVInfo
 {
     class Version
     {
-        private static string VersionString = "1.02";
+        private static string VersionString = "1.03";
 
         public static string Get() { return VersionString; }
     };
@@ -186,6 +186,27 @@ namespace CSVInfo
 
 
 			// Set metadata if requested
+			string setMetadataFile = GetArg("setMetadataFile");
+			if (setMetadataFile != null)
+			{
+				if (csvStats.metaData == null)
+				{
+					throw new Exception("File has no metadata, so can't set it");
+				}
+				string [] metadataLines=System.IO.File.ReadAllLines(setMetadataFile);
+				foreach (string line in metadataLines)
+				{
+					int equalsIndex = line.IndexOf("=");
+					if (equalsIndex != -1)
+					{
+						string Key = line.Substring(0, equalsIndex);
+						string Value = line.Substring(equalsIndex + 1);
+						Console.WriteLine("Setting metadata: " + Key + "=" + Value);
+						csvStats.metaData.Values[Key] = Value;
+					}
+				}
+			}
+
 			string SetMetadataStr = GetArg("setMetadata");
 			if (SetMetadataStr != null && SetMetadataStr.Length > 0)
 			{
