CsvInfo 1.03: -quiet mode. Useful if you only care about the return code (ie to determine if a CSV is valid)

diff --git a/Engine/Source/Programs/CSVTools/CsvInfo/CSVInfo.cs b/Engine/Source/Programs/CSVTools/CsvInfo/CSVInfo.cs
--- a/Engine/Source/Programs/CSVTools/CsvInfo/CSVInfo.cs
+++ b/Engine/Source/Programs/CSVTools/CsvInfo/CSVInfo.cs
@@ -12,7 +12,7 @@ namespace CSVInfo
 {
     class Version
     {
-        private static string VersionString = "1.02";
+        private static string VersionString = "1.03";
 
         public static string Get() { return VersionString; }
     };
@@ -73,6 +73,7 @@ namespace CSVInfo
                 "Format: \n" +
                 "  <csvfilename>\n"+
 				"  [-showaverages]\n"+
+				"  [-quiet] (no logging. Just throws returns a non-zero error code if the CSV is bad)\n" +
 				"  [-toJson <filename>]";
 
 			// Read the command line
@@ -91,6 +92,12 @@ namespace CSVInfo
 			string jsonFilename = GetArg("toJson",false);
 
 			CsvStats csvStats = CsvStats.ReadCSVFile(csvFilename, null);
+
+			if ( GetBoolArg("quiet") )
+			{
+				return;
+			}
+
 			List<string> statLines = new List<string>();
 			foreach (StatSamples stat in csvStats.Stats.Values.ToArray())
 			{
