CsvConvert - Fix exception which occurs if setMetadataFile isn't passed on the commandline

diff --git a/Engine/Source/Programs/CSVTools/CsvConvert/CsvConvert.cs b/Engine/Source/Programs/CSVTools/CsvConvert/CsvConvert.cs
--- a/Engine/Source/Programs/CSVTools/CsvConvert/CsvConvert.cs
+++ b/Engine/Source/Programs/CSVTools/CsvConvert/CsvConvert.cs
@@ -13,7 +13,7 @@ namespace CSVInfo
 {
     class Version
     {
-        private static string VersionString = "1.03";
+        private static string VersionString = "1.04";
 
         public static string Get() { return VersionString; }
     };
@@ -186,7 +186,7 @@ namespace CSVInfo
 
 
 			// Set metadata if requested
-			string setMetadataFile = GetArg("setMetadataFile");
+			string setMetadataFile = GetArg("setMetadataFile",null);
 			if (setMetadataFile != null)
 			{
 				if (csvStats.metaData == null)
