CsvToSvg 2.50. Refactor code to move graph generation into a separate lib, so we can use it from PerfReportTool directly instead of executing a separate process and reading all the Csv data multiple times. This also cleans up the code in CsvToSvg  considerably.
- Graph Generation is now handled by a GraphGenerator class with MakeGraphs async and synchronous methods. Settings are passed via a GraphParams struct.
- Optimize various functions to speed up graph generation.
Generating graphs for a report is now around 25% faster.

diff --git a/Engine/Source/Programs/CSVTools/CSVToSVG/CSVToSVG.csproj b/Engine/Source/Programs/CSVTools/CSVToSVG/CSVToSVG.csproj
--- a/Engine/Source/Programs/CSVTools/CSVToSVG/CSVToSVG.csproj
+++ b/Engine/Source/Programs/CSVTools/CSVToSVG/CSVToSVG.csproj
@@ -32,6 +32,10 @@
     <ErrorReport>prompt</ErrorReport>
     <WarningLevel>4</WarningLevel>
   </PropertyGroup>
+  <PropertyGroup>
+    <StartupObject>
+    </StartupObject>
+  </PropertyGroup>
   <ItemGroup>
     <Reference Include="System" />
     <Reference Include="System.Core" />
@@ -42,7 +46,7 @@
     <Reference Include="System.Xml" />
   </ItemGroup>
   <ItemGroup>
-    <Compile Include="CsvToSVG.cs" />
+    <Compile Include="CsvToSvgProgram.cs" />
     <Compile Include="Properties\AssemblyInfo.cs" />
   </ItemGroup>
   <ItemGroup>
@@ -53,6 +57,10 @@
       <Project>{75c595cc-83ba-4ffd-a772-17b2343d1061}</Project>
       <Name>CSVStats</Name>
     </ProjectReference>
+    <ProjectReference Include="..\CsvToSvgLib\CsvToSvgLib.csproj">
+      <Project>{75c595cc-83ba-4ffd-abcd-17b2343d1061}</Project>
+      <Name>CsvToSvgLib</Name>
+    </ProjectReference>
   </ItemGroup>
   <Import Project="$(MSBuildToolsPath)\Microsoft.CSharp.targets" />
   <!-- To modify your build process, add your task inside one of the targets below and uncomment it. 
diff --git a/Engine/Source/Programs/CSVTools/CSVToSVG/CsvToSVG.cs b/Engine/Source/Programs/CSVTools/CSVToSVG/CsvToSVG.cs
deleted file mode 100644
diff --git a/Engine/Source/Programs/CSVTools/CSVToSVG/CsvToSvgProgram.cs b/Engine/Source/Programs/CSVTools/CSVToSVG/CsvToSvgProgram.cs
new file mode 100644
--- /dev/null
+++ b/Engine/Source/Programs/CSVTools/CSVToSVG/CsvToSvgProgram.cs
@@ -0,0 +1,482 @@
+// Copyright (C) Microsoft. All rights reserved.
+// Copyright Epic Games, Inc. All Rights Reserved.
+
+using System;
+using System.Collections.Generic;
+using System.Linq;
+using System.Text;
+using System.Threading.Tasks;
+using System.IO;
+using System.Diagnostics;
+using System.Security.Cryptography;
+using CSVStats;
+
+namespace CSVTools
+{
+	class CsvCache
+	{
+		public CsvStats ReadCSVFile(string csvFilename, string[] statNames, int numRowsToSkip)
+		{
+			lock (readLock)
+			{
+				if (csvFilename != filename || skipRows != numRowsToSkip)
+				{
+					csvStats = CsvStats.ReadCSVFile(csvFilename, null, numRowsToSkip);
+					filename = csvFilename;
+					skipRows = numRowsToSkip;
+				}
+				// Copy just the stats we're interested in from the cached file
+				CsvStats outStats = new CsvStats(csvStats, statNames);
+				return outStats;
+			}
+		}
+		public CsvStats csvStats;
+		public string filename;
+		public int skipRows;
+		private readonly object readLock = new object();
+	};
+
+
+	class Program : CSVStats.CommandLineTool
+    {
+		static CsvCache csvCache = null;
+
+		static string formatString =
+            "Format: \n" +
+            "       -csvs <list> OR -csv <list> OR -csvDir <path>\n" +
+            "       [ -o <svgFilename> ]\n" +
+            "       -stats <stat names> (can include wildcards)\n" +
+			"     OR \n" +
+			"       -batchCommands <response file with commandlines>\n" +
+			"       -mt <number of threads> \n" +
+			"     OR \n" +
+			"       -updatesvg <svgFilename>\n" +
+			"       NOTE: this updates an svg by regenerating it with the original commandline parameters - requires original data\n\n" +
+			"\nOptional Args: \n" +
+			"       -averageThreshold <value>\n" +
+			"       -budget <ms>\n" +
+			"       -colourOffset <value>\n" +
+			"       -compression <pixel error value>\n" +
+			"       -discardLastFrame <1|0> (default 1)\n" +
+			"       -filterOutZeros\n" +
+			"       -graphOnly\n" +
+			"       -hideEventNames <1|0>\n" +
+			"       -hideStatPrefix <list>\n" +
+			"       -hierarchySeparator <character>\n" +
+			"       -highlightEventRegions <startEventName,endEventName>\n" +
+			"       -ignoreStats <list> (can include wildcards. Separate states with stat1;stat2;etc)\n" +
+			"       -interactive\n" +
+			"       -legend <list> \n" +
+			"       -maxHierarchyDepth <depth>\n" +
+			"       -minX <value> -maxX <value> -minY <value> -maxY <value>\n" +
+			"       -noMetadata\n" +
+			"       -noSnap\n" +
+			"       -recurse\n" +
+			"       -showAverages \n" +
+			"       -showEvents <names> (can include wildcards)\n" +
+			"       -showTotals \n" +
+			"       -skipRows <n>\n" +
+			"       -smooth\n" +
+			"       -smoothKernelPercent <percentage>\n" +
+			"       -smoothKernelSize <numFrames>\n" +
+			"       -smoothMultithreaded <1|0>\n"+
+			"       -stacked\n" +
+			"       -stackTotalStat <stat name>\n" +
+			"		-minFilterStatValue <value>\n" +
+			"		-minFilterStatName <stat name>\n" +
+			"       -stackedUnsorted\n" +
+			"       -statMultiplier <multiplier>\n" +
+			"       -theme <dark|light>\n" +
+			"       -thickness <multipler>\n" +
+			"       -threshold <value>\n" +
+			"       -title <name>\n" +
+			"       -forceLegendSort\n" +
+			"       -width <value> -height <value>\n" +
+            "       -writeErrorsToSVG\n" +
+            "       -percentile \n" +
+            "       -percentile90 \n" +
+            "       -percentile99 \n" +
+			"       -uniqueId <string> : unique ID for JS (needed if this is getting embedded in HTML alongside other graphs)\n" +
+			"       -nocommandlineEmbed : don't embed the commandline in the SVG" +
+			"       -lineDecimalPlaces <N> (default 3)" +
+			"       -frameOffset <N> : offset used for frame display name (default 0)" +
+            "";
+
+		void Run(string[] args)
+		{
+			System.Globalization.CultureInfo.DefaultThreadCurrentCulture = System.Globalization.CultureInfo.InvariantCulture;
+
+			if (args.Length < 2)
+			{
+				WriteLine("CsvToSVG " + CsvToSvgLibVersion.Get());
+				WriteLine(formatString);
+				return;
+			}
+
+			// Read the command line
+			ReadCommandLine(args);
+
+			string svgToUpdate = GetArg("updatesvg");
+			if (svgToUpdate.Length > 0)
+			{
+				if (args.Length > 2)
+				{
+					WriteLine("-UpdateSVG <svgFilename> must be the only argument!");
+					return;
+				}
+				string newCommandLine = "";
+				string[] svgLines = ReadLinesFromFile(svgToUpdate);
+				for (int i = 0; i < svgLines.Length - 2; i++)
+				{
+					if (svgLines[i].StartsWith("<![CDATA[") && svgLines[i + 1].StartsWith("Created with CSVtoSVG ") && svgLines[i + 1].EndsWith(" with commandline:"))
+					{
+						newCommandLine = svgLines[i + 2];
+						break;
+					}
+				}
+
+				ReadCommandLine(MakeArgsArray(newCommandLine));
+			}
+
+			MakeGraph();
+		}
+
+		string GenerateIDFromString(string str)
+		{
+			string id = "";
+			HashAlgorithm algorithm = SHA256.Create();
+			byte[] hash = algorithm.ComputeHash(Encoding.UTF8.GetBytes(commandLine.GetCommandLine()));
+			for ( int i=24; i<32; i++ )
+			{
+				id += hash[i].ToString("x2");
+			}
+			return id;
+		}
+
+
+		void MakeGraph()
+		{
+			// Read CSV filenames from a directory or list
+			string[] csvFilenames;
+			string csvDir = GetArg("csvDir");
+			if (csvDir.Length > 0)
+			{
+				DirectoryInfo di = new DirectoryInfo(csvDir);
+				bool recurse = GetBoolArg("recurse");
+				FileInfo[] csvFiles = di.GetFiles("*.csv", recurse ? SearchOption.AllDirectories : SearchOption.TopDirectoryOnly);
+				FileInfo[] binFiles = di.GetFiles("*.csv.bin", recurse ? SearchOption.AllDirectories : SearchOption.TopDirectoryOnly);
+				List<FileInfo> allFiles = new List<FileInfo>(csvFiles);
+				allFiles.AddRange(binFiles);
+				csvFilenames = new string[allFiles.Count];
+				int i = 0;
+				foreach (FileInfo csvFile in allFiles)
+				{
+					csvFilenames[i] = csvFile.FullName;
+					i++;
+				}
+			}
+			else
+			{
+				string csvFilenamesStr = GetArg("csv");
+				if (csvFilenamesStr.Length == 0)
+				{
+					csvFilenamesStr = GetArg("csvs", true);
+					if (csvFilenamesStr.Length == 0)
+					{
+						System.Console.Write(formatString);
+						return;
+					}
+				}
+				csvFilenames = csvFilenamesStr.Split(';');
+			}
+
+			GraphParams graphParams = getGraphParamsFromArgs();
+
+			if (graphParams.statNames.Count == 0)
+			{
+				System.Console.Write(formatString);
+				return;
+			}
+
+			if (graphParams.statNames.Count > 1 && csvFilenames.Length > 1)
+			{
+				WriteLine("Can't display multiple stats and multiple CSVs");
+				return;
+			}
+
+			// Figure out the filename, if it wasn't provided
+			string svgFilename = GetArg("o", false);
+			if (svgFilename.Length == 0)
+			{
+				if (csvFilenames.Length == 1)
+				{
+					int index = csvFilenames[0].LastIndexOf('.');
+					if (index >= 0)
+					{
+						svgFilename = csvFilenames[0].Substring(0, index);
+					}
+					else
+					{
+						svgFilename = csvFilenames[0];
+					}
+				}
+				if (graphParams.statNames.Count == 1)
+				{
+					if (svgFilename.Length > 0)
+					{
+						svgFilename += "_";
+					}
+					svgFilename += graphParams.statNames[0].Replace('*', 'X').Replace('/', '_');
+				}
+				svgFilename += ".svg";
+			}
+			if (svgFilename.Length == 0)
+			{
+				WriteLine("Couldn't figure out an appropriate filename, and no filename was provided with -o");
+			}
+
+			// Load the CSVs
+			List<CsvInfo> csvs = new List<CsvInfo>();
+			foreach (string csvFilename in csvFilenames)
+			{
+				CsvStats csvStats = LoadCsv(csvFilename, graphParams.statNames.ToArray());
+				csvs.Add(new CsvInfo(csvStats, csvFilename));
+			}
+
+			// Read graph params from the commandline
+			GraphGenerator generator = new GraphGenerator(csvs);
+
+			bool writeErrorsToSVG = GetIntArg("writeErrorsToSVG", 1) == 1;
+			try
+			{
+				generator.MakeGraph(graphParams, svgFilename, writeErrorsToSVG);
+			}
+			catch (System.Exception e)
+			{
+				// Write the error to the SVG
+				string errorString = e.ToString();
+				Console.Out.WriteLine("Error: " + e.ToString());
+			}
+		}
+
+        static string[] MakeArgsArray(string argsStr)
+        {
+            List<string> argsOut = new List<string>();
+            string currentArg = "";
+            bool bInQuotes = false;
+            for ( int i=0; i< argsStr.Length;i++ )
+            {
+                bool flush = false;
+                char c = argsStr[i];
+                if ( c == '"')
+                {
+                    if ( bInQuotes )
+                    {
+                        flush = true;
+                    }
+                    bInQuotes = !bInQuotes;
+                }
+                else if ( c==' ' )
+                {
+                    if ( bInQuotes )
+                    {
+                        currentArg += c;
+                    }
+                    else
+                    {
+                        flush = true;
+                    }
+                }
+                else
+                {
+                    currentArg += c;
+                }
+
+                if (flush && currentArg.Length > 0)
+                {
+                    argsOut.Add(currentArg);
+                    currentArg = "";
+                }
+            }
+			if (currentArg.Length > 0)
+			{
+				argsOut.Add(currentArg);
+			}
+			return argsOut.ToArray();
+        }
+
+
+		GraphParams getGraphParamsFromArgs()
+		{
+			GraphParams graphParams = new GraphParams();
+
+			graphParams.width = GetIntArg("width", graphParams.width);
+			graphParams.height = GetIntArg("height", graphParams.height);
+			graphParams.title = GetArg("title", graphParams.title);
+			graphParams.statNames = GetListArg("stats", ';', false, true);
+			graphParams.ignoreStats = GetListArg("ignoreStats");
+			graphParams.themeName = GetArg("theme").ToLower();
+			graphParams.colorOffset = GetIntArg("colorOffset", GetIntArg("colourOffset", 0) );
+
+			// Events
+			graphParams.showEventNames = GetListArg("showEvents");
+			graphParams.showEventNameText = GetIntArg("hideEventNames", 0) == 0;
+			graphParams.highlightEventRegions = GetListArg("highlightEventRegions", ',');
+
+
+			// Smoothing
+			graphParams.smooth = GetArg("smooth") == "1";
+			graphParams.smoothKernelPercent = GetFloatArg("smoothKernelPercent", graphParams.smoothKernelPercent);
+			graphParams.smoothKernelSize = GetIntArg("smoothKernelSize", graphParams.smoothKernelSize);
+
+			// Other flags
+			graphParams.showMetadata = GetArg("nometadata") != "1";
+			graphParams.graphOnly = GetArg("graphOnly") == "1";
+			graphParams.compression = GetFloatArg("compression", graphParams.compression);
+			graphParams.interactive = GetArg("interactive") == "1";
+			graphParams.snapToPeaks = !GetBoolArg("noSnap");
+
+			// Sometimes the last frame is garbage. We might want to remove it
+			graphParams.discardLastFrame = (GetIntArg("discardLastFrame", graphParams.discardLastFrame?1:0) == 1);
+
+			graphParams.maxHierarchyDepth = GetIntArg("maxHierarchyDepth", graphParams.maxHierarchyDepth);
+
+			string hierarchySeparatorStr = GetArg("hierarchySeparator","");
+			if (hierarchySeparatorStr.Length > 0)
+			{
+				graphParams.hierarchySeparator = hierarchySeparatorStr[0];
+			}
+			graphParams.hideStatPrefixes = GetListArg("hideStatPrefix", ';', true);
+
+			// Stacking
+			graphParams.stacked = GetArg("stacked") == "1";
+			if (graphParams.stacked)
+			{
+				graphParams.stackTotalStat = GetArg("stackTotalStat").ToLower();
+			}
+			graphParams.stackedUnsorted = GetArg("stackedUnsorted") == "1";
+
+			// Percentiles
+			graphParams.percentileTop90 = GetArg("percentile90") == "1";
+			graphParams.percentileTop99 = GetArg("percentile99") == "1";
+			graphParams.percentile = GetArg("percentile") == "1";
+
+			// X/Y range
+			graphParams.minX = GetFloatArg("minx", graphParams.minX);
+			graphParams.maxX = GetFloatArg("maxx", graphParams.maxX);
+			graphParams.minY = GetFloatArg("miny", graphParams.minY);
+			graphParams.maxY = GetFloatArg("maxy", graphParams.maxY);
+
+			graphParams.budget = GetFloatArg("budget", graphParams.budget);
+			graphParams.lineThickness = GetFloatArg("thickness", graphParams.lineThickness);
+
+			// Thresholds
+			graphParams.threshold = GetFloatArg("threshold", graphParams.threshold);
+			graphParams.averageThreshold = GetFloatArg("averageThreshold", graphParams.averageThreshold);
+
+			graphParams.embedText = "Created with CSVtoSVG " + CsvToSvgLibVersion.Get();
+			if (!GetBoolArg("nocommandlineEmbed"))
+			{
+				graphParams.embedText += " with commandline:\n" + commandLine.GetCommandLine();
+			}
+			graphParams.uniqueId = GetArg("uniqueID", graphParams.uniqueId);
+
+			// Legend
+			graphParams.showAverages = GetBoolArg("showAverages");
+			graphParams.showTotals = GetBoolArg("showTotals");
+			graphParams.forceLegendSort = GetBoolArg("forceLegendSort");
+			graphParams.legendAverageThreshold = GetFloatArg("legendAverageThreshold", graphParams.legendAverageThreshold);
+
+			graphParams.customLegendNames = GetListArg("legend");
+
+			// Process options (don't affect output)
+			string smoothMultiThreadedStr = GetArg("smoothMultithreaded",null);
+			if (smoothMultiThreadedStr != null)
+			{
+				graphParams.bSmoothMultithreaded = smoothMultiThreadedStr == "1";
+			}
+
+			// Misc params (rarely used)
+			graphParams.lineDecimalPlaces=GetIntArg("lineDecimalPlaces", graphParams.lineDecimalPlaces);
+			graphParams.frameOffset = GetIntArg("frameOffset", graphParams.frameOffset);
+			graphParams.statMultiplier = GetFloatArg("statMultiplier", graphParams.statMultiplier);
+
+			graphParams.minFilterStatName = GetArg("minFilterStatName", null);
+			if (graphParams.minFilterStatName != null)
+			{
+				graphParams.minFilterStatValue = GetFloatArg("minFilterStatValue", -float.MaxValue);
+			}
+			graphParams.filterOutZeros = GetBoolArg("filterOutZeros");
+
+			graphParams.FinalizeSettings();
+			return graphParams;
+		}
+
+		CsvStats LoadCsv(string csvFilename, string[] statNames)
+        {
+            int numRowsToSkip = GetIntArg("skipRows",0);
+
+			// Read the file from the cache if we have one
+			CsvStats csvStats = null;
+			if (csvCache != null)
+			{
+				csvStats = csvCache.ReadCSVFile(csvFilename, statNames, numRowsToSkip);
+			}
+			else
+			{
+				csvStats = CsvStats.ReadCSVFile(csvFilename, statNames, numRowsToSkip);
+			}
+            return csvStats;
+        }
+
+		static void Main(string[] args)
+        {
+			CommandLine commandline = new CommandLine(args);
+			string batchCommandsFilename = commandline.GetArg("batchCommands");
+
+			if (batchCommandsFilename.Length > 0)
+			{
+				// In batch mode, just output a list of commandlines
+				string[] commandlines = System.IO.File.ReadAllLines(batchCommandsFilename);
+				if (commandlines.Length > 1)
+				{
+					// Enable caching
+					csvCache = new CsvCache();
+				}
+				int numThreads = commandline.GetIntArg("mt", 4);
+				if (numThreads>1)
+				{
+					var result = Parallel.For(0, commandlines.Length, new ParallelOptions { MaxDegreeOfParallelism = numThreads }, i =>
+					{
+						Program program = new Program();
+						program.Run(MakeArgsArray(commandlines[i]));
+					});
+				}
+				else
+				{
+					foreach (string cmdLine in commandlines)
+					{
+						Program program = new Program();
+						program.Run(MakeArgsArray(cmdLine));
+					}
+				}
+			}
+			else
+			{
+				Program program = new Program();
+				try
+				{
+					program.Run(args);
+				}
+				catch (System.Exception e)
+				{
+					Console.WriteLine("[ERROR] " + e.Message);
+					if (Debugger.IsAttached)
+					{
+						throw e;
+					}
+				}
+			}
+        }
+     }
+}
diff --git a/Engine/Source/Programs/CSVTools/CsvToSvgLib/CsvToSvgLib.cs b/Engine/Source/Programs/CSVTools/CsvToSvgLib/CsvToSvgLib.cs
new file mode 100644
--- /dev/null
+++ b/Engine/Source/Programs/CSVTools/CsvToSvgLib/CsvToSvgLib.cs
@@ -0,0 +1,2376 @@
+﻿using System;
+using System.Collections.Generic;
+using System.Linq;
+using CSVStats;
+using System.Threading.Tasks;
+using System.Text;
+using System.ComponentModel;
+
+namespace CSVTools
+{
+
+	public class CsvToSvgLibVersion
+	{
+		private static string VersionString = "2.50";
+
+		public static string Get() { return VersionString; }
+	};
+
+	public class GraphParams
+	{
+		public int width = 1800;
+		public int height = 550;
+		public string title = "";
+		public List<string> statNames = new List<string>();
+		public List<string> ignoreStats = new List<string>();
+		public int colorOffset = 0;
+
+		// Events
+		public List<string> showEventNames = new List<string>();
+		public bool showEventNameText = true;
+		public List<string> highlightEventRegions = new List<string>();
+
+		// Smoothing
+		public bool smooth = false;
+		public float smoothKernelPercent = 2.0f;
+		public int smoothKernelSize = -1;
+
+		// Other flags
+		public bool showMetadata = true; //noMetadata
+		public bool graphOnly = false;
+		public float compression = 0.0f;
+		public bool interactive = false;
+		public bool snapToPeaks = true; // Interactive mode only
+		public int maxHierarchyDepth = -1;
+		public char hierarchySeparator = '/';
+		public List<string> hideStatPrefixes = new List<string>();
+
+		// Stacking
+		public bool stacked = false;
+		public string stackTotalStat = "";
+		public bool stackedUnsorted = false;
+
+		// Percentiles
+		public bool percentileTop90 = false;
+		public bool percentileTop99 = false;
+		public bool percentile = false;
+
+		// X/Y range
+		public float minX = Range.Auto;
+		public float maxX = Range.Auto;
+		public float minY = Range.Auto;
+		public float maxY = Range.Auto;
+
+		public float budget = float.MaxValue;
+		public float lineThickness = 1.0f;
+
+		// Thresholds
+		public float threshold = -float.MaxValue;
+		public float averageThreshold = -float.MaxValue;
+
+		public string embedText = "Created with CSVtoSVG "+ CsvToSvgLibVersion.Get();
+		public string uniqueId = "ID";
+
+		// Legend
+		public bool showAverages = false;
+		public bool showTotals = false;
+		public bool forceLegendSort = false;
+		public float legendAverageThreshold = -float.MaxValue;
+		public List<string> customLegendNames = new List<string>();
+
+		// Process options (doesn't affect output)
+		public bool bSmoothMultithreaded = true;
+		public bool bPerfLog = false;
+
+		// Advanced params
+		public string themeName = "";
+		public int lineDecimalPlaces = 3;
+		public int frameOffset = 0;
+		public float statMultiplier = 1.0f;
+		public string minFilterStatName = null;
+		public float minFilterStatValue = -float.MaxValue;
+		public bool filterOutZeros = false;
+		public bool discardLastFrame = true;
+
+		private bool bFinalized = false;
+
+		public void FinalizeSettings()
+		{
+			if (bFinalized)
+			{
+				return;
+			}
+
+			if (percentile && (stacked || interactive))
+			{
+				throw new Exception("Warning: percentile graph not compatible with stacked & interactive");
+			}
+
+			// Make sure stack total stat is in the stat list. This may result in a duplicate, but that's fine
+			if (stacked)
+			{
+				if (stackTotalStat != "")
+				{
+					stackTotalStat = stackTotalStat.ToLower();
+					statNames.Add(stackTotalStat);
+				}
+			}
+
+			if ( percentileTop90 || percentileTop99 )
+			{
+				percentile = true;
+			}
+
+			themeName = themeName.ToLower();
+
+			if (maxX <= 0.0f)
+			{
+				maxX = Range.Auto;
+			}
+			if (maxY <= 0.0f)
+			{
+				maxY = Range.Auto;
+			}
+			if (smoothKernelSize>0.0f)
+			{
+				smoothKernelPercent = 0.0f;
+			}
+			bFinalized = true;
+		}
+
+
+		public void DebugPrintProperties()
+		{
+			foreach (PropertyDescriptor descriptor in TypeDescriptor.GetProperties(this))
+			{
+				string name = descriptor.Name;
+				object value = descriptor.GetValue(this);
+				Console.WriteLine("{0}={1}", name, value);
+			}
+		}
+	};
+
+	public class CsvInfo
+	{
+		public CsvInfo(CsvStats csvStatsIn, string csvFilenameIn)
+		{
+			stats = csvStatsIn;
+			filename = csvFilenameIn;
+		}
+		public CsvStats stats;
+		public string filename;
+	};
+
+	class SvgFile
+	{
+		public SvgFile(string filename, string inUniqueId=null)
+		{
+			file = new System.IO.StreamWriter(filename);
+			idSuffix = "_" + inUniqueId;
+			sb = new StringBuilder();
+		}
+
+		public void WriteFast(string str)
+		{
+			sb.Append(str);
+		}
+
+		public void Write(string str)
+		{
+			sb.Append(str.Replace("<UNIQUE>", idSuffix));
+		}
+
+		public void WriteLineFast(string str)
+		{
+			sb.Append(str+"\n");
+		}
+
+		public void WriteLine(string str)
+		{
+			sb.Append(str.Replace("<UNIQUE>", idSuffix)+"\n");
+			Flush();
+		}
+
+		public void Flush()
+		{
+			file.Write(sb.ToString());
+			sb.Clear();
+		}
+
+		public void Close()
+		{
+			Flush();
+			file.Close();
+			file = null;
+		}
+
+		readonly string idSuffix;
+		System.IO.StreamWriter file = null;
+		StringBuilder sb = null;
+		//int bufferLineCount = 0;
+	}
+
+
+
+	class Rect
+	{
+		public Rect(float xIn, float yIn, float widthIn, float heightIn)
+		{
+			x = xIn; y = yIn; width = widthIn; height = heightIn;
+		}
+		public float x, y;
+		public float width, height;
+	};
+
+	class Range
+	{
+		public const float Auto = -100000.0f;
+
+		public Range() { }
+		public Range(float minx, float maxx, float miny, float maxy) { MinX = minx; MaxX = maxx; MinY = miny; MaxY = maxy; }
+		public float MinX, MaxX;
+		public float MinY, MaxY;
+	};
+
+	class Theme
+	{
+		// TODO: make this data driven
+		public Theme(string Name)
+		{
+			GraphColours = new Colour[32];
+			uint[] GraphColoursInt = null;
+			if (Name == "light")
+			{
+				BackgroundColour = new Colour(255, 255, 255);
+				BackgroundColourCentre = new Colour(255, 255, 255);
+				LineColour = new Colour(0, 0, 0);
+
+				GraphColoursInt = new uint[16]
+				{
+					0x0000C0, 0x8000C0, 0xFF4000, 0xC00000,
+					0x4040A0, 0x008080, 0x200080, 0x408060,
+					0x008040, 0x00008C, 0x60A880, 0x325000,
+					0xA040A0, 0x808000, 0x005050, 0x606060,
+				};
+
+				TextColour = new Colour(0, 0, 0);
+				MediumTextColour = new Colour(64, 64, 64);
+				MinorTextColour = new Colour(128, 128, 128);
+				AxisLineColour = new Colour(128, 128, 128);
+				MajorGridlineColour = new Colour(128, 128, 128);
+				MinorGridlineColour = new Colour(160, 160, 160);
+				BudgetLineColour = new Colour(0, 196, 0);
+				EventTextColour = new Colour(0);
+				BudgetLineThickness = 1.0f;
+			}
+			else if (Name == "pink")
+			{
+				BackgroundColour = new Colour(255, 128, 128);
+				BackgroundColourCentre = new Colour(255, 255, 255);
+				LineColour = new Colour(0, 0, 0);
+				GraphColoursInt = new uint[16]
+				{
+					0x8080FF, 0xFF8C8C, 0xFFFF8C, 0x20C0C0,
+					0x808000, 0xFF8C8C, 0x20FF8C, 0x408060,
+					0xFF8040, 0xFFFF8C, 0x60008C, 0x3250FF,
+					0x008000, 0x8C8CFF, 0xFF5050, 0x606060,
+				};
+
+				TextColour = new Colour(0, 0, 0);
+				MediumTextColour = new Colour(192, 192, 192);
+				MinorTextColour = new Colour(128, 128, 128);
+				AxisLineColour = new Colour(128, 128, 128);
+				MajorGridlineColour = new Colour(128, 128, 128);
+				MinorGridlineColour = new Colour(160, 160, 160);
+				BudgetLineColour = new Colour(0, 196, 0);
+				EventTextColour = new Colour(0);
+				BudgetLineThickness = 1.0f;
+			}
+			else // "dark"
+			{
+				BackgroundColour = new Colour(16, 16, 16);
+				BackgroundColourCentre = new Colour(80, 80, 80);
+				LineColour = new Colour(255, 255, 255);
+
+				GraphColoursInt = new uint[16]
+				{//0x11bbbb
+                    0x0080FF, 0x66cdFF, 0xFF6600, 0xFFFF8C,
+					0x60f060, 0xFFFF00, 0x99CC00, 0xCC6600,
+					0xCC3300, 0xCCFF66, 0x60008C, 0x3250FF,
+					0x008000, 0x11bbbb, 0xFF5050, 0x606060,
+				};
+				for (int i = 0; i < 16; i++)
+					GraphColours[i] = new Colour(GraphColoursInt[i]);
+
+				TextColour = new Colour(255, 255, 255);
+				MediumTextColour = new Colour(192, 192, 192);
+				MinorTextColour = new Colour(128, 128, 128);
+				AxisLineColour = new Colour(128, 128, 128);
+				MajorGridlineColour = new Colour(128, 128, 128);
+				MinorGridlineColour = new Colour(96, 96, 96);
+				BudgetLineColour = new Colour(128, 255, 128);
+				EventTextColour = new Colour(255, 255, 255, 0.75f);
+				BudgetLineThickness = 0.5f;
+			}
+
+			for (int i = 0; i < GraphColours.Length; i++)
+			{
+				int repeat = i / GraphColoursInt.Length;
+				float alpha = 1.0f - (float)repeat * 0.25f;
+				GraphColours[i] = new Colour(GraphColoursInt[i % GraphColoursInt.Length], alpha);
+
+			}
+
+		}
+		public Colour BackgroundColour;
+		public Colour BackgroundColourCentre;
+		public Colour LineColour;
+		public Colour TextColour;
+		public Colour MinorTextColour;
+		public Colour MediumTextColour;
+		public Colour AxisLineColour;
+		public Colour MinorGridlineColour;
+		public Colour MajorGridlineColour;
+		public Colour BudgetLineColour;
+		public Colour EventTextColour;
+		public float BudgetLineThickness;
+		public Colour[] GraphColours;
+	};
+
+
+	public class GraphGenerator
+	{
+		List<CsvInfo> csvList;
+
+		public GraphGenerator(List<CsvInfo> csvsIn)
+		{
+			csvList = csvsIn;
+		}
+
+		public GraphGenerator(CsvStats csv, string filename)
+		{
+			csvList = new List<CsvInfo>();
+			csvList.Add(new CsvInfo(csv, filename));
+		}
+
+		public Task MakeGraphAsync(GraphParams graphParams, string svgFilename, bool bWriteErrorsToSvg, bool rethrowExceptions = true)
+		{
+			Action action = () =>
+			{
+				MakeGraph(graphParams, svgFilename, bWriteErrorsToSvg, rethrowExceptions);
+			};
+			return Task.Factory.StartNew(action);
+		}
+
+		public void MakeGraph(GraphParams graphParams, string svgFilename, bool bWriteErrorsToSvg, bool rethrowExceptions = true)
+		{
+			SvgFile svg = new SvgFile(svgFilename, graphParams.uniqueId);
+			try
+			{
+				graphParams.FinalizeSettings();
+				MakeGraphInternal(graphParams, svg);
+			}
+			catch (System.Exception e)
+			{
+				// Write the error to the SVG
+				string errorString = e.ToString();
+				if (bWriteErrorsToSvg)
+				{
+					errorString = errorString.Replace(" at", " at<br/>\n");
+					errorString += "<br/><br/>"+ graphParams.embedText.Replace("\n","<br/>");
+					float MessageWidth = graphParams.width - 20;
+					float MessageHeight = graphParams.height - 20;
+
+					svg.WriteLine("<switch>");
+					svg.WriteLine("<foreignObject x='10' y='10' color='#ffffff' font-size='12' width='" + MessageWidth + "' height='" + MessageHeight + "'><p xmlns='http://www.w3.org/1999/xhtml'>" + errorString + "</p></foreignObject>'");
+					svg.WriteLine("<text x='10' y='10' fill='rgb(255, 255, 255)' font-size='10' font-family='Helvetica' > ERROR: " + errorString + "</text>");
+					svg.WriteLine("</switch>");
+					svg.WriteLine("</svg>");
+				}
+				svg.Close();
+				if (rethrowExceptions)
+				{
+					throw;
+				}
+			}
+			svg.Close();
+		}
+
+		void MakeGraphInternal(GraphParams graphParams, SvgFile svg)
+		{
+			PerfLog perfLog = new PerfLog(graphParams.bPerfLog, "Graph:"+graphParams.title);
+
+			Rect dimensions = new Rect(0, 0, graphParams.width, graphParams.height);
+			Rect graphRect = new Rect(50, 42, dimensions.width - 100, dimensions.height - 115);
+			Theme theme = new Theme(graphParams.themeName);
+			List<string> statNames = new List<string>(graphParams.statNames);
+			string stackTotalStat = null;
+			bool stackTotalStatIsAutomatic = false;
+
+			if (graphParams.stacked)
+			{
+				stackTotalStat = graphParams.stackTotalStat;
+			}
+
+			svg.Write("<svg width='" + dimensions.width + "' height='" + dimensions.height + "' viewPort='0 0 " + dimensions.height + " " + dimensions.width + "' version='1.1' xmlns='http://www.w3.org/2000/svg'");
+			if (graphParams.interactive)
+			{
+				svg.Write(" onLoad='OnLoaded<UNIQUE>(evt)'");
+			}
+			svg.WriteLine(">");
+
+			if (graphParams.embedText != null)
+			{
+				svg.WriteLine("<![CDATA[ \n"+graphParams.embedText);
+				svg.WriteLine("]]>");
+			}
+
+			// Write defs
+			svg.WriteLine("<defs>");
+			svg.WriteLine("<radialGradient id='radialGradient1'");
+			svg.WriteLine("fx='50%' fy='50%' r='65%'");
+			svg.WriteLine("spreadMethod='pad'>");
+			svg.WriteLine("<stop offset='0%'   stop-color=" + theme.BackgroundColourCentre.SVGString() + " stop-opacity='1'/>");
+			svg.WriteLine("<stop offset='100%' stop-color=" + theme.BackgroundColour.SVGString() + " stop-opacity='1' />");
+			svg.WriteLine("</radialGradient>");
+
+			svg.WriteLine("<linearGradient id = 'linearGradient1' x1 = '0%' y1 = '0%' x2 = '100%' y2 = '100%'>");
+			svg.WriteLine("<stop stop-color = 'black' offset = '0%'/>");
+			svg.WriteLine("<stop stop-color = 'white' offset = '100%'/>");
+			svg.WriteLine("</linearGradient>");
+
+			svg.WriteLine("<clipPath id='graphArea'>");
+			svg.WriteLine("<rect  x='" + graphRect.x + "' y='" + graphRect.y + "' width='" + graphRect.width + "' height='" + graphRect.height + "'/>");
+			svg.WriteLine("</clipPath>");
+
+			svg.WriteLine("<filter id='dropShadowFilter' x='-20%' width='130%' height='130%'>");
+			svg.WriteLine("<feOffset result='offOut' in='SourceAlpha' dx='-2' dy='2' />");
+			svg.WriteLine("<feGaussianBlur result='blurOut' in='offOut' stdDeviation='1.1' />");
+			svg.WriteLine("<feBlend in='SourceGraphic' in2='blurOut' mode='normal' />");
+			svg.WriteLine("</filter>");
+
+			svg.WriteLine("</defs>");
+
+			
+			DrawGraphArea(svg, theme, theme.BackgroundColour, dimensions, true);
+			perfLog.LogTiming("Initialization");
+
+			// Generate the list of processed, filtered CSV stats for graphing
+			List<CsvStats> csvStatsList = new List<CsvStats>();
+			foreach (CsvInfo csvInfo in csvList)
+			{
+				CsvStats newCsvStats = ProcessCsvStats(csvInfo.stats, graphParams);
+
+				if (graphParams.stacked && stackTotalStat == "")
+				{
+					// Make a total stat by summing each frame
+					StatSamples totalStat = new StatSamples("Total");
+					totalStat.samples.Capacity = newCsvStats.SampleCount;
+					for (int i = 0; i < newCsvStats.SampleCount; i++)
+					{
+						float totalValue = 0.0f;
+						foreach (StatSamples stat in newCsvStats.Stats.Values)
+						{
+							totalValue += stat.samples[i];
+						}
+						totalStat.samples.Add(totalValue);
+					}
+					totalStat.ComputeAverageAndTotal();
+					totalStat.colour = new Colour(0x6E6E6E);
+					newCsvStats.AddStat(totalStat);
+					stackTotalStat = "total";
+					stackTotalStatIsAutomatic = true;
+				}
+
+				SetLegend(newCsvStats, csvInfo.filename, graphParams, csvList.Count > 1);
+
+				if (graphParams.smooth)
+				{
+					int kernelSize = graphParams.smoothKernelSize;
+					if (graphParams.smoothKernelSize == -1)
+					{
+						float percent = graphParams.smoothKernelPercent;
+						percent = Math.Min(percent, 100.0f);
+						percent = Math.Max(percent, 0.0f);
+						float kernelSizeF = percent * 0.01f * (float)newCsvStats.SampleCount + 0.5f;
+						kernelSize = (int)kernelSizeF;
+					}
+					newCsvStats = SmoothStats(newCsvStats, kernelSize, graphParams.bSmoothMultithreaded);
+				}
+				if (!graphParams.stacked)
+				{
+					AssignColours(newCsvStats, theme, false, graphParams.colorOffset);
+				}
+				csvStatsList.Add(newCsvStats);
+			}
+			perfLog.LogTiming("ProcessCsvStats");
+
+			Range range = new Range(graphParams.minX, graphParams.maxX, graphParams.minY, graphParams.maxY);
+
+			// Compute the X range 
+			range = ComputeAdjustedXRange(range, graphRect, csvStatsList);
+
+			// Recompute the averages based on the X range. This is necessary for the legend and accurate sorting
+			for (int i = 0; i < csvStatsList.Count; i++)
+			{
+				RecomputeStatAveragesForRange(csvStatsList[i], range);
+			}
+			perfLog.LogTiming("RecomputeStatAverages");
+
+
+			// Handle stacking. Note that if we're not stacking, unstackedCsvStats will be a copy of csvStats
+			List<CsvStats> unstackedCsvStats = new List<CsvStats>();
+			for (int i = 0; i < csvStatsList.Count; i++)
+			{
+				unstackedCsvStats.Add(csvStatsList[i]);
+
+				if (graphParams.stacked)
+				{
+					csvStatsList[i] = StackStats(csvStatsList[i], range, graphParams.stackedUnsorted, stackTotalStat, stackTotalStatIsAutomatic);
+					AssignColours(csvStatsList[i], theme, true, graphParams.colorOffset);
+
+					// Copy the colours to the unstacked stats
+					foreach (StatSamples samples in unstackedCsvStats[i].Stats.Values)
+					{
+						string statName = samples.Name;
+						if (stackTotalStatIsAutomatic == false && stackTotalStat != null && samples.Name.ToLower() == stackTotalStat.ToLower())
+						{
+							statName = "other";
+						}
+						samples.colour = csvStatsList[i].GetStat(statName).colour;
+					}
+				}
+			}
+
+			if (graphParams.stacked) perfLog.LogTiming("StackStats");
+
+			// Adjust range
+			range = ComputeAdjustedYRange(range, graphRect, csvStatsList);
+
+			if (graphParams.graphOnly)
+			{
+				graphRect = dimensions;
+			}
+
+			// Adjust thickness depending on sample density, smoothness etc
+			float thickness = graphParams.smooth ? 0.33f : 0.1f;
+			float thicknessMultiplier = (12000.0f / (range.MaxX - range.MinX)) * graphParams.lineThickness;
+			thicknessMultiplier *= (graphRect.width / 400.0f);
+			thickness *= thicknessMultiplier;
+			thickness = Math.Max(Math.Min(thickness, 1.5f), 0.11f);
+
+			string graphTitle = graphParams.title;
+
+			// Get the title
+			if (graphTitle.Length == 0 && statNames.Count == 1 && csvStatsList.Count > 0 && !statNames[0].Contains("*"))
+			{
+				StatSamples stat = csvStatsList[0].GetStat(statNames[0]);
+				if (stat == null)
+				{
+					Console.Out.WriteLine("Warning: Could not find stat {0}", statNames[0]);
+					graphTitle = string.Format("UnknownStat {0}", statNames[0]);
+				}
+				else
+				{
+					graphTitle = stat.Name;
+				}
+			}
+
+			// Combine and validate metadata
+			// Assign metadata based on the first CSV with metadata
+			CsvMetadata metaData = null;
+			
+			if (graphParams.showMetadata && !graphParams.graphOnly)
+			{
+				foreach (CsvStats stats in csvStatsList)
+				{
+					if (stats.metaData != null)
+					{
+						metaData = stats.metaData.Clone();
+						break;
+					}
+				}
+				if (metaData != null)
+				{
+					// Combine all metadata
+					foreach (CsvStats stats in csvStatsList)
+					{
+						metaData.CombineAndValidate(stats.metaData);
+					}
+				}
+			}
+			perfLog.LogTiming("AdjustRangesAndMetadata");
+
+			// Draw the graphs
+			int csvIndex = 0;
+			uint colourIndex = 0;
+			svg.WriteLine("<g id='graphArea'>");
+			if (graphParams.percentile)
+			{
+				range.MinX = graphParams.percentileTop99 ? 99 : (graphParams.percentileTop90 ? 90 : 0);
+				range.MaxX = 100;
+				DrawGridLines(svg, theme, graphRect, range, graphParams, !graphParams.graphOnly, 1.0f, true);
+				csvIndex = 0;
+				foreach (CsvStats csvStat in csvStatsList)
+				{
+					foreach (StatSamples stat in csvStat.Stats.Values)
+					{
+						string statID = "Stat_" + csvIndex + "_" + stat.Name + "<UNIQUE>";
+						DrawPercentileGraph(svg, stat.samples, stat.colour, graphRect, range, statID);
+						colourIndex++;
+					}
+					csvIndex++;
+				}
+			}
+			else
+			{
+				DrawGridLines(svg, theme, graphRect, range, graphParams, !graphParams.graphOnly, 1.0f, graphParams.stacked);
+				csvIndex = 0;
+				foreach (CsvStats csvStat in csvStatsList)
+				{
+					DrawEventLines(svg, theme, csvStat.Events, graphRect, range);
+					DrawEventHighlightRegions(svg, csvStat, graphRect, range, graphParams.highlightEventRegions);
+
+					foreach (StatSamples stat in csvStat.Stats.Values)
+					{
+						string statID = "Stat_" + csvIndex + "_" + stat.Name + "<UNIQUE>";
+						DrawGraph(svg, stat.samples, stat.colour, graphRect, range, thickness, statID, graphParams);
+						colourIndex++;
+					}
+
+					if (graphParams.showEventNameText)
+					{
+						Colour eventColour = theme.EventTextColour;
+						DrawEventText(svg, csvStat.Events, eventColour, graphRect, range);
+					}
+
+					csvIndex++;
+				}
+			}
+			perfLog.LogTiming("DrawGraphs");
+
+
+			svg.WriteLine("</g>");
+
+			if (graphParams.stacked)
+			{
+				// If we're stacked, we need to redraw the grid lines without axis text
+				DrawGridLines(svg, theme, graphRect, range, graphParams, false, 0.75f, false);
+			}
+
+			// Draw legend, metadata and title
+			if (!graphParams.graphOnly)
+			{
+				DrawLegend(svg, theme, csvStatsList, graphParams, dimensions);
+				if (graphParams.smooth)
+				{
+					DrawText(svg, "(smoothed)", 50.0f, 33, 10.0f, dimensions, theme.MinorTextColour, "start");
+				}
+				DrawTitle(svg, theme, graphTitle, dimensions);
+			}
+
+
+			if (metaData != null)
+			{
+				DrawMetadata(svg, metaData, dimensions, theme.TextColour);
+			}
+
+			DrawText(svg, "CSVToSVG " + CsvToSvgLibVersion.Get(), dimensions.width - 102, 15, 7.0f, dimensions, theme.MinorTextColour);
+
+			perfLog.LogTiming("DrawText");
+
+			// Draw the interactive elements
+			if (graphParams.interactive)
+			{
+				AddInteractiveScripting(svg, theme, graphRect, range, unstackedCsvStats, graphParams);
+				perfLog.LogTiming("AddInteractiveScripting");
+			}
+
+			svg.WriteLine("</svg>");
+			perfLog.LogTotalTiming(false);
+		}
+
+		Range ComputeAdjustedYRange(Range range, Rect rect, List<CsvStats> csvStats)
+		{
+			Range newRange = new Range(range.MinX, range.MaxX, range.MinY, range.MaxY);
+			int maxNumSamples = 0;
+			float maxSample = -10000000.0f;
+			float minSample = 10000000.0f;
+
+			foreach (CsvStats stats in csvStats)
+			{
+				foreach (StatSamples samples in stats.Stats.Values)
+				{
+					maxNumSamples = Math.Max(maxNumSamples, samples.samples.Count);
+					foreach (float sample in samples.samples)
+					{
+						if (sample > maxSample)
+						{
+							maxSample = sample;
+						}
+						minSample = Math.Min(minSample, sample);
+						maxSample = Math.Max(maxSample, sample);
+					}
+				}
+			}
+
+			// HACK: Clamp to 1m to prevent craziness
+			if (maxSample > 1000000) maxSample = 1000000;
+
+			if (range.MinY == Range.Auto) newRange.MinY = 0.0f;
+			if (range.MaxY == Range.Auto) newRange.MaxY = maxSample * 1.05f;
+
+			// Quantise based on yincrement
+			if (rect != null)
+			{
+				float yInc = GetYAxisIncrement(rect, newRange);
+				float difY = newRange.MaxY - newRange.MinY;
+				float newDifY = (int)(0.9999 + difY / yInc) * yInc;
+				newRange.MaxY = newRange.MinY + newDifY;
+			}
+
+			return newRange;
+		}
+
+		Range ComputeAdjustedXRange(Range range, Rect rect, List<CsvStats> csvStats)
+		{
+			Range newRange = new Range(range.MinX, range.MaxX, range.MinY, range.MaxY);
+			int maxNumSamples = 0;
+			foreach (CsvStats stats in csvStats)
+			{
+				foreach (StatSamples samples in stats.Stats.Values)
+				{
+					maxNumSamples = Math.Max(maxNumSamples, samples.samples.Count);
+				}
+			}
+
+			if (range.MinX == Range.Auto) newRange.MinX = 0;
+			if (range.MaxX == Range.Auto) newRange.MaxX = maxNumSamples;
+
+			// Quantize based on xincrement
+			if (rect != null)
+			{
+				float xInc = GetXAxisIncrement(rect, newRange);
+				float difX = newRange.MaxX - newRange.MinX;
+				float newDifX = (int)(0.9999 + difX / xInc) * xInc;
+				newRange.MaxX = newRange.MinX + newDifX;
+			}
+
+			return newRange;
+		}
+
+
+
+		CsvStats SmoothStats(CsvStats stats, int KernelSize, bool bMultiThreaded)
+		{
+			// Compute Gaussian Weights
+			float[] Weights = new float[KernelSize];
+			float TotalWeight = 0.0f;
+			for (int i = 0; i < KernelSize; i++)
+			{
+				double df = (double)i + 0.5;
+				double b = (double)KernelSize / 2;
+				double c = (double)KernelSize / 8;
+				double weight = Math.Exp(-(Math.Pow(df - b, 2.0) / Math.Pow(2 * c, 2.0)));
+				Weights[i] = (float)weight;
+				TotalWeight += (float)weight;
+			}
+
+			for (int i = 0; i < KernelSize; i++)
+			{
+				Weights[i] /= TotalWeight;
+			}
+
+			List<StatSamples> statSampleList = stats.Stats.Values.ToList();
+
+			// Add the stats to smoothstats before the parallel for, so we can preserve the order
+			CsvStats smoothStats = new CsvStats();
+			foreach (StatSamples srcStatSamples in stats.Stats.Values)
+			{
+				StatSamples NewStatSamples = new StatSamples(srcStatSamples, false);
+				smoothStats.AddStat(NewStatSamples);
+			}
+			StatSamples[] SmoothSamplesArray = smoothStats.Stats.Values.ToArray();
+
+			Action<int> smoothStat = i =>
+			{
+				StatSamples srcStatSamples = statSampleList[i];
+				StatSamples destStatSamples = SmoothSamplesArray[i];// new StatSamples(srcStatSamples);
+
+				int maxSampleIndex = srcStatSamples.samples.Count - 1;
+				for (int j = 0; j < srcStatSamples.samples.Count; j++)
+				{
+					float sum = 0.0f;
+					for (int k = 0; k < KernelSize; k++)
+					{
+						int SampleIndex = j - (KernelSize / 2) + k;
+
+						// If the index is out of range, mirror it (so as not to bias the result)
+						if (SampleIndex < 0)
+						{
+							SampleIndex = -SampleIndex;
+						}
+						else if (SampleIndex > maxSampleIndex)
+						{
+							SampleIndex = maxSampleIndex - (SampleIndex - maxSampleIndex);
+						}
+						// clamp the src index to a valid range
+						int srcIndex = Math.Min(Math.Max(0, SampleIndex), maxSampleIndex);
+
+						sum += srcStatSamples.samples[srcIndex] * Weights[k];
+					}
+					destStatSamples.samples.Add(sum);
+				}
+			};
+
+			if (bMultiThreaded)
+			{
+				Parallel.For(0, stats.Stats.Values.Count, smoothStat);
+			}
+			else
+			{
+				for (int i=0; i< stats.Stats.Values.Count; i++)
+				{
+					smoothStat(i);
+				}
+			}
+
+
+			smoothStats.metaData = stats.metaData;
+			smoothStats.Events = stats.Events;
+			return smoothStats;
+		}
+
+		void SetLegend(CsvStats csvStats, string csvFilename, GraphParams graphParams, bool UseFilename)
+		{
+			int customIndex = 0;
+			foreach (StatSamples stat in csvStats.Stats.Values)
+			{
+				stat.LegendName = stat.Name;
+				if (UseFilename)
+				{
+					stat.LegendName = System.IO.Path.GetFileName(csvFilename);
+				}
+				if (graphParams.customLegendNames.Count > 0 && customIndex < graphParams.customLegendNames.Count)
+				{
+					stat.LegendName = graphParams.customLegendNames[customIndex];
+					customIndex++;
+				}
+				foreach (string hideStatPrefix in graphParams.hideStatPrefixes)
+				{
+					if (hideStatPrefix.Length > 0)
+					{
+						if (stat.LegendName.ToLower().StartsWith(hideStatPrefix.ToLower()))
+						{
+							stat.LegendName = stat.LegendName.Substring(hideStatPrefix.Length);
+						}
+					}
+				}
+			}
+		}
+
+		void AssignColours(CsvStats stats, Theme theme, bool reverseOrder, int inColorOffset)
+		{
+			int colorOffset = inColorOffset;
+			if (reverseOrder)
+			{
+				colorOffset = stats.Stats.Values.Count - 1;
+			}
+
+			foreach (StatSamples stat in stats.Stats.Values)
+			{
+				if (stat.colour.alpha == 0.0f)
+				{
+					stat.colour = theme.GraphColours[colorOffset % theme.GraphColours.Length];
+				}
+				if (reverseOrder)
+				{
+					colorOffset--;
+				}
+				else
+				{
+					colorOffset++;
+				}
+			}
+		}
+		CsvStats StackStats(CsvStats stats, Range range, bool stackedUnsorted, string stackTotalStat, bool bStackTotalStatIsAutomatic)
+		{
+			// Find our stack total stat
+			StatSamples totalStat = null;
+			string stackTotalStatLower = stackTotalStat.ToLower();
+			if (stackTotalStat.Length > 0)
+			{
+				foreach (StatSamples stat in stats.Stats.Values)
+				{
+					if (stat.Name.ToLower() == stackTotalStatLower)
+					{
+						totalStat = stat;
+						break;
+					}
+				}
+			}
+
+			// sort largest->smallest (based on average)
+			List<StatSamples> StatSamplesSorted = new List<StatSamples>();
+			foreach (StatSamples stat in stats.Stats.Values)
+			{
+				if (stat != totalStat)
+				{
+					StatSamplesSorted.Add(stat);
+				}
+			}
+
+			if (!stackedUnsorted)
+			{
+				StatSamplesSorted.Sort();
+			}
+
+			StatSamples OtherStat = null;
+			StatSamples OtherUnstacked = null;
+			int rangeStart = Math.Max(0, (int)range.MinX);
+			int rangeEnd = (int)range.MaxX;
+
+			if (!bStackTotalStatIsAutomatic && totalStat != null)
+			{
+				OtherStat = new StatSamples("Other");
+				OtherStat.colour = new Colour(0x6E6E6E);
+				StatSamplesSorted.Add(OtherStat);
+				OtherStat.samples.AddRange(totalStat.samples);
+				OtherUnstacked = new StatSamples(OtherStat, false);
+				rangeEnd = Math.Min(OtherStat.samples.Count, (int)range.MaxX);
+			}
+
+			// Get the max count for all stats
+			int maxCount = 0;
+			foreach (StatSamples statSamples in StatSamplesSorted)
+			{
+				maxCount = Math.Max(statSamples.samples.Count, maxCount);
+			}
+
+			// Make the stats cumulative
+			float[] SumList = new float[maxCount];
+			List<StatSamples> StatSamplesList = new List<StatSamples>();
+			foreach (StatSamples srcStatSamples in StatSamplesSorted)
+			{
+				StatSamples destStatSamples = new StatSamples(srcStatSamples, false);
+				if (srcStatSamples == OtherStat)
+				{
+					for (int j = 0; j < srcStatSamples.samples.Count; j++)
+					{
+						// The previous value of value is the total. If there are no other stats, then the total is "other"
+						float value = Math.Max(srcStatSamples.samples[j] - SumList[j], 0.0f);
+						OtherUnstacked.samples.Add(value);
+						SumList[j] += value;
+						destStatSamples.samples.Add(SumList[j]);
+					}
+					// If this is the "other" stat, recompute the average based on the unstacked value and apply that to the dest stat (this is used in the legend)
+					OtherUnstacked.ComputeAverageAndTotal(rangeStart, rangeEnd);
+					destStatSamples.average = OtherUnstacked.average;
+				}
+				else 
+				{
+					for (int j = 0; j < srcStatSamples.samples.Count; j++)
+					{
+						SumList[j] += srcStatSamples.samples[j];
+						destStatSamples.samples.Add(SumList[j]);
+					}
+				}
+				StatSamplesList.Add(destStatSamples);
+			}
+
+			// Copy out the list in reverse order (so they get rendered front->back)
+			CsvStats stackedStats = new CsvStats();
+			if (bStackTotalStatIsAutomatic)
+			{
+				stackedStats.AddStat(totalStat);
+			}
+
+			for (int i = StatSamplesList.Count - 1; i >= 0; i--)
+			{
+				stackedStats.AddStat(StatSamplesList[i]);
+			}
+
+			stackedStats.metaData = stats.metaData;
+			stackedStats.Events = stats.Events;
+			return stackedStats;
+		}
+
+		float GetXAxisIncrement(Rect rect, Range range)
+		{
+			float[] xIncrements = { 0.1f, 1.0f, 5.0f, 10.0f, 20.0f, 50.0f, 100.0f, 200.0f, 500.0f, 1000.0f, 2000.0f, 5000.0f, 10000.0f };
+			float xIncrement = xIncrements[0];
+
+			for (int i = 0; i < xIncrements.Length - 1; i++)
+			{
+				float gap = ToSvgXScale(xIncrement, rect, range);
+				if (gap < 25)
+				{
+					xIncrement = xIncrements[i + 1];
+				}
+				else
+				{
+					break;
+				}
+
+			}
+			int xRange = (int)range.MaxX - (int)range.MinX;
+			//xIncrement = xIncrement - (xRange % (int)xIncrement);
+			return xIncrement;
+		}
+
+		float GetYAxisIncrement(Rect rect, Range range)
+		{
+			float[] yIncrements = { 0.5f, 1.0f, 2.0f, 5.0f, 10.0f, 20.0f, 50.0f, 100.0f, 200.0f, 500.0f, 1000.0f, 2000, 5000, 10000, 20000, 50000, 100000, 200000, 500000, 1000000, 2000000, 5000000, 10000000 };
+
+			// Find a good increment
+			float yIncrement = yIncrements[0];
+			for (int i = 0; i < yIncrements.Length - 1; i++)
+			{
+				float gap = ToSvgYScale(yIncrement, rect, range);
+				if (gap < 25)
+				{
+					yIncrement = yIncrements[i + 1];
+				}
+				else
+				{
+					break;
+				}
+			}
+			return yIncrement;
+		}
+
+		void DrawGridLines(SvgFile svg, Theme theme, Rect rect, Range range, GraphParams graphParams, bool bDrawAxisText, float alpha = 1.0f, bool extendLines = false)
+		{
+			float xIncrement = GetXAxisIncrement(rect, range);
+			float yIncrement = GetYAxisIncrement(rect, range);
+			float yScale = 1.0f;
+
+			// Draw 1ms grid lines
+			float minorYAxisIncrement = 1.0f;
+			if ((range.MaxY - range.MinY) > 100) minorYAxisIncrement = yIncrement / 10.0f;
+			minorYAxisIncrement = Math.Max(minorYAxisIncrement, yIncrement / 5.0f);
+
+			Colour colour = new Colour(theme.MinorGridlineColour);
+			colour.alpha = alpha;
+
+			// TODO: replace with <pattern>
+			svg.WriteLine("<g id='gridLines<UNIQUE>'>");
+			for (int i = 0; i < 5000; i++)
+			{
+				float y = range.MinY + (float)i * minorYAxisIncrement;
+				DrawHorizLine(svg, y, colour, rect, range);
+				if (y >= range.MaxY) break;
+			}
+
+			colour = new Colour(theme.MajorGridlineColour);
+			colour.alpha = (alpha + 3.0f) / 4.0f;
+			for (int i = 0; i < 2000; i++)
+			{
+				float y = range.MinY + (float)i * yIncrement;
+				DrawHorizLine(svg, y, colour, rect, range, false, 0.25f, false);
+				yScale += yIncrement;
+				if (y >= range.MaxY) break;
+			}
+			svg.WriteLine("</g>");
+
+			// Draw the axis text
+			if (bDrawAxisText)
+			{
+				svg.WriteLine("<g id='xAxis<UNIQUE>'>");
+				for (int i = 0; i < 2000; i++)
+				{
+					float x = range.MinX + i * xIncrement;
+					if (x > (range.MaxX)) break;
+					int displayFrame = (int)Math.Floor(x) + graphParams.frameOffset;
+					DrawHorizontalAxisText(svg, displayFrame.ToString(), x, theme.TextColour, rect, range);
+				}
+				svg.WriteLine("</g>");
+
+				svg.WriteLine("<g id='yAxis<UNIQUE>'>");
+				for (int i = 0; i < 2000; i++)
+				{
+					float y = range.MinY + (float)i * yIncrement;
+					// FIXME: This text should use theme.LineColour - using LineColour for now for consistency/historical reasons. Need to test 
+					DrawVerticalAxisText(svg, y.ToString(), y, theme.LineColour, rect, range);
+					DrawVerticalAxisNotch(svg, y, theme.LineColour, rect, range);
+					if (y >= range.MaxY) break;
+				}
+				svg.WriteLine("</g>");
+			}
+
+			DrawVerticalLine(svg, 0, theme.AxisLineColour, rect, range, 0.5f);
+			DrawHorizLine(svg, 0.0f, theme.AxisLineColour, rect, range, false, 0.5f);
+
+			// Draw the budget line if there is one
+			if (graphParams.budget != float.MaxValue)
+			{
+				float budgetLineThickness = theme.BudgetLineThickness;
+				Colour budgetLineColour = new Colour(theme.BudgetLineColour);
+				bool dropShadow = false;
+				if (alpha != 1.0f)
+				{
+					dropShadow = true;
+					budgetLineThickness *= 2.0f;
+					budgetLineColour.r = (byte)((float)budgetLineColour.r * 0.9f);
+					budgetLineColour.g = (byte)((float)budgetLineColour.g * 0.9f);
+					budgetLineColour.b = (byte)((float)budgetLineColour.b * 0.9f);
+				}
+				DrawHorizLine(svg, graphParams.budget, budgetLineColour, rect, range, true, budgetLineThickness, dropShadow);
+			}
+		}
+
+		void DrawText(SvgFile svg, string text, float x, float y, float size, Rect rect, Colour colour, string anchor = "start", string font = "Helvetica", string id = "", bool dropShadow = false)
+		{
+			svg.WriteLine("<text x='" + (rect.x + x) + "' y='" + (rect.y + y) + "' fill=" + colour.SVGString() +
+					  (id.Length > 0 ? " id='" + id + "'" : "") +
+					  (dropShadow ? " filter='url(#dropShadowFilter)'" : "") +
+					  " font-size='" + size + "' font-family='" + font + "' style='text-anchor: " + anchor + "'>" + text + "</text> >");
+		}
+
+		void DrawTitle(SvgFile svg, Theme theme, string title, Rect rect)
+		{
+			DrawText(svg, title, (float)(rect.x + rect.width * 0.5), 34, 14, rect, theme.TextColour, "middle", "Helvetica");
+		}
+
+
+		void DrawGraphArea(SvgFile svg, Theme theme, Colour colour, Rect rect, bool gradient, bool interactive = false)
+		{
+			//svg.WriteLine("<rect x='0' y='0' width='" + (dimensions.width) + "' height='" + (dimensions.height) + "' fill=" + colour.SVGString() + " stroke-width='1' stroke=" + theme.BackgroundColour.SVGString());
+			svg.WriteLine("<rect x='0' y='0' width='" + (rect.width) + "' height='" + (rect.height) + "' fill=" + colour.SVGString() + " stroke-width='1' stroke=" + theme.BackgroundColour.SVGString());
+
+			if (gradient)
+			{
+				svg.Write("style='fill:url(#radialGradient1);' ");
+			}
+			if (interactive)
+			{
+				svg.Write(" onmousemove='OnGraphAreaClicked<UNIQUE>(evt)'");
+				//svg.Write(" onmousewheel='OnGraphAreaWheeled<UNIQUE>(evt)'");
+			}
+			svg.Write("/>");
+		}
+
+		float ToSvgX(float graphX, Rect rect, Range range)
+		{
+			float scaleX = rect.width / (range.MaxX - range.MinX);
+			return rect.x + ((float)graphX - range.MinX) * scaleX;
+		}
+
+		float ToSvgY(float graphY, Rect rect, Range range)
+		{
+			float svgY = rect.height - (graphY - range.MinY) / (range.MaxY - range.MinY) * rect.height + rect.y;
+
+			float scaleY = rect.height / (range.MaxY - range.MinY);
+			float svgY2 = rect.y + rect.height - (graphY - range.MinY) * scaleY;
+			return svgY;
+		}
+
+		float ToSvgXScale(float graphX, Rect rect, Range range)
+		{
+			float scaleX = rect.width / (range.MaxX - range.MinX);
+			return graphX * scaleX;
+		}
+
+		float ToSvgYScale(float graphY, Rect rect, Range range)
+		{
+			float scaleY = rect.height / (range.MaxY - range.MinY);
+			return graphY * scaleY;
+		}
+
+		float ToCsvXScale(float graphX, Rect rect, Range range)
+		{
+			float scaleX = rect.width / (range.MaxX - range.MinX);
+			return graphX / scaleX;
+		}
+
+		float ToCsvYScale(float graphY, Rect rect, Range range)
+		{
+			float scaleY = rect.height / (range.MaxY - range.MinY);
+			return graphY / scaleY;
+		}
+
+		float ToCsvX(float svgX, Rect rect, Range range)
+		{
+			return (svgX - rect.x) * (range.MaxX - range.MinX) / rect.width + range.MinX;
+		}
+
+		float ToCsvY(float svgY, Rect rect, Range range)
+		{
+			return (svgY - rect.y) * (range.MaxY - range.MinY) / rect.height + range.MinY;
+		}
+
+		struct Vec2
+		{
+			public Vec2(float inX, float inY) { X = inX; Y = inY; }
+			public float X;
+			public float Y;
+		};
+
+		class PointInfo
+		{
+			public void AddPoint(Vec2 pt)
+			{
+				Count++;
+				Max.X = Math.Max(pt.X, Max.X);
+				Max.Y = Math.Max(pt.Y, Max.Y);
+				Min.X = Math.Min(pt.X, Min.X);
+				Min.Y = Math.Min(pt.Y, Min.Y);
+				Total.X += pt.X;
+				Total.Y += pt.Y;
+			}
+			public Vec2 ComputeAvg()
+			{
+				Vec2 avg;
+				avg.X = Total.X / (float)Count;
+				avg.Y = Total.Y / (float)Count;
+				return avg;
+			}
+			public Vec2 Max = new Vec2(float.MinValue, float.MinValue);
+			public Vec2 Min = new Vec2(float.MaxValue, float.MaxValue);
+			public Vec2 Total = new Vec2(0, 0);
+			public int Count = 0;
+		};
+
+		Vec2[] RemoveRedundantPoints(Vec2[] RawPoints, float compressionThreshold, int passIndex)
+		{
+			//return RawPoints;
+			//compressionThreshold = 0;
+			if (RawPoints.Length == 0)
+			{
+				return RawPoints;
+			}
+			List<Vec2> StrippedPoints = new List<Vec2>(RawPoints.Length);
+
+			int offset = passIndex % 2;
+
+			// Add the first points
+			int i = 0;
+			for (i = 0; i <= offset; i++)
+			{
+				StrippedPoints.Add(RawPoints[i]);
+			}
+			for (; i < RawPoints.Length - 1; i += 2)
+			{
+				Vec2 p1 = RawPoints[i - 1];
+				Vec2 pCentre = RawPoints[i];
+				Vec2 p2 = RawPoints[i + 1];
+
+				// Interpolate to find where the line would be if we removed this point
+				float dX = p2.X - p1.X;
+				float centreDX = pCentre.X - p1.X;
+				float lerpValue = centreDX / dX;
+				float interpX = p1.X * (1.0f - lerpValue) + (p2.X * lerpValue);
+				float interpY = p1.Y * (1.0f - lerpValue) + (p2.Y * lerpValue);
+
+				//float distMh = Math.Abs(interpX - pCentre.X) + Math.Abs(interpY - pCentre.Y);
+				float distMh = Math.Abs(interpY - pCentre.Y);
+				if (distMh >= compressionThreshold)
+				{
+					StrippedPoints.Add(pCentre);
+				}
+				StrippedPoints.Add(p2);
+			}
+			// Add the last points
+			for (; i < RawPoints.Length; i++)
+			{
+				StrippedPoints.Add(RawPoints[i]);
+			}
+			return StrippedPoints.ToArray();
+		}
+
+		List<Vec2> CompressPoints(List<Vec2> RawPoints, float pixelsPerPoint, Rect rect, Range range)
+		{
+			if (RawPoints.Count < 3)
+			{
+				return RawPoints;
+			}
+			int numPointsAfterCompression = (int)((float)rect.width / pixelsPerPoint);
+			if (numPointsAfterCompression >= RawPoints.Count)
+			{
+				return RawPoints;
+			}
+			float NumOutputPointsPerInputPoint = (float)numPointsAfterCompression / (float)RawPoints.Count;
+
+			PointInfo[] pointInfoList = new PointInfo[numPointsAfterCompression];
+			for (int i = 0; i < numPointsAfterCompression; i++) pointInfoList[i] = new PointInfo();
+
+			// Compute average, min, max for every point
+			for (int i = 0; i < RawPoints.Count; i++)
+			{
+				Vec2 point = RawPoints[i];
+				float outPointX = (float)i * NumOutputPointsPerInputPoint;
+				int outPointIndex = Math.Min((int)(outPointX), numPointsAfterCompression - 1);
+				pointInfoList[outPointIndex].AddPoint(point);
+			}
+
+			Vec2[] CompressedPoints = new Vec2[numPointsAfterCompression];
+			CompressedPoints[0] = pointInfoList[0].ComputeAvg();
+			for (int i = 1; i < pointInfoList.Length - 1; i++)
+			{
+				PointInfo info = pointInfoList[i];
+				PointInfo last = pointInfoList[i - 1];
+				PointInfo next = pointInfoList[i + 1];
+				if (info.Min.Y <= last.Min.Y && info.Min.Y < next.Min.Y)
+				{
+					CompressedPoints[i] = info.Min;
+				}
+				else if (info.Max.Y >= last.Max.Y && info.Max.Y > next.Max.Y)
+				{
+					CompressedPoints[i] = info.Max;
+				}
+				else
+				{
+					CompressedPoints[i] = info.ComputeAvg();
+				}
+			}
+			CompressedPoints[pointInfoList.Length - 1]=pointInfoList.Last().ComputeAvg();
+			float compressionThreshold = pixelsPerPoint;
+			for (int i = 0; i < 8; i++)
+			{
+				CompressedPoints = RemoveRedundantPoints(CompressedPoints, compressionThreshold, i);
+			}
+			return CompressedPoints.ToList();
+		}
+
+		void DrawGraph(SvgFile svg, List<float> samples, Colour colour, Rect rect, Range range, float thickness, string id, GraphParams graphParams)
+		{
+			int n = Math.Min(samples.Count, (int)(range.MaxX + 0.5));
+			int start = (int)(range.MinX + 0.5);
+
+			int numDecimalPlaces = graphParams.lineDecimalPlaces;
+
+			if (start + 1 < n)
+			{
+				List<Vec2> RawPoints = new List<Vec2>(n);
+				for (int i = start; i < n; i++)
+				{
+					float sample = samples[i];
+					float x = ToSvgX((float)i, rect, range);
+					float y = ToSvgY(sample, rect, range);
+
+					RawPoints.Add(new Vec2(x, y));
+				}
+
+				if (graphParams.compression > 0.0f)
+				{
+					RawPoints = CompressPoints(RawPoints, graphParams.compression, rect, range);
+				}
+
+				string formatStr = "0";
+				if (numDecimalPlaces > 0)
+				{
+					formatStr += ".";
+					for (int i = 0; i < numDecimalPlaces; i++)
+					{
+						formatStr += "0";
+					}
+				}
+				string idString = id.Length > 0 ? "id='" + id + "'" : "";
+
+				svg.WriteLine("<polyline " + idString + " points='");
+				foreach (Vec2 point in RawPoints)
+				{
+					float x = point.X;
+					float y = point.Y;
+					svg.WriteFast(" " + x.ToString(formatStr) + "," + y.ToString(formatStr));
+				}
+
+				string fillString = "none";
+				if (graphParams.stacked)
+				{
+					thickness = 0.0f;// Math.Min(thickness, 1.0f);
+					float lastSample = samples[n - 1];
+					svg.WriteFast(" " + ToSvgX((float)n + 20, rect, range) + "," + ToSvgY(lastSample, rect, range) );
+					svg.WriteFast(" " + ToSvgX((float)n + 20, rect, range) + "," + ToSvgY(0.0f, rect, range) );
+					svg.WriteFast(" " + ToSvgX((float)n - 20, rect, range) + "," + ToSvgY(0.0f, rect, range) );
+					svg.WriteFast(" " + ToSvgX(start, rect, range) + "," + ToSvgY(0.0f, rect, range) );
+
+					colour.alpha = 1.0f;
+					fillString = colour.SVGStringNoQuotes();
+				}
+				svg.WriteLine("' style='fill:" + fillString + ";stroke-width:" + thickness + "; clip-path: url(#graphArea)' stroke=" + colour.SVGString() + "/>");
+			}
+		}
+		void DrawPercentileGraph(SvgFile svg, List<float> samples, Colour colour, Rect rect, Range range, string id)
+		{
+			samples.Sort();
+			List<Vec2> RawPoints = new List<Vec2>();
+			const int NUM_SAMPLES = 100;
+			int Begin = (int)range.MinX;
+			int End = (int)range.MaxX;
+			int Delta = End - Begin;
+			float Percentile, xbase, x, y;
+			int Count = samples.Count;
+			for (int i = NUM_SAMPLES * Begin; i < NUM_SAMPLES * End; i += Delta)
+			{
+				int sampleIndex = Count * i / (NUM_SAMPLES * 100);
+				Percentile = samples[sampleIndex];
+				xbase = (float)i / NUM_SAMPLES;
+				x = ToSvgX(xbase, rect, range);
+				y = ToSvgY(Percentile, rect, range);
+				RawPoints.Add(new Vec2(x, y));
+			}
+
+			Percentile = samples[Count - 1];
+			xbase = 100.0F;
+			x = ToSvgX(xbase, rect, range);
+			y = ToSvgY(Percentile, rect, range);
+			RawPoints.Add(new Vec2(x, y));
+			string idString = id.Length > 0 ? "id='" + id + "'" : "";
+			svg.WriteLine("<polyline " + idString + " points='");
+			foreach (Vec2 point in RawPoints)
+			{
+				x = point.X;
+				y = point.Y;
+
+				svg.WriteFast(" " + x + "," + y);
+			}
+			svg.WriteLine("' style='fill:none;stroke-width:1.3; clip-path: url(#graphArea)' stroke=" + colour.SVGString() + "/>");
+
+		}
+
+
+		class CsvEventWithCount : CsvEvent
+		{
+			public CsvEventWithCount(CsvEvent ev)
+			{
+				base.Frame = ev.Frame;
+				base.Name = ev.Name;
+				count = 1;
+			}
+			public int count;
+		};
+
+		void DrawEventText(SvgFile svg, List<CsvEvent> events, Colour colour, Rect rect, Range range)
+		{
+			float LastEventX = -100000.0f;
+			int lastFrame = 0;
+
+			// Make a filtered list of events to display, grouping duplicates which are close together
+			List<CsvEventWithCount> filteredEvents = new List<CsvEventWithCount>();
+			CsvEventWithCount currentDisplayEvent = null;
+			CsvEvent lastEvent = null;
+			foreach (CsvEvent ev in events)
+			{
+				// Only draw events which are in the range
+				if (ev.Frame >= range.MinX && ev.Frame <= range.MaxX)
+				{
+					// Merge with the current display event?
+					bool bMerge = false;
+					if (currentDisplayEvent != null && lastEvent != null && ev.Name == currentDisplayEvent.Name)
+					{
+						float DistToLastEvent = ToSvgX(ev.Frame, rect, range) - ToSvgX(lastEvent.Frame, rect, range);
+						if (DistToLastEvent <= 8.5)
+						{
+							bMerge = true;
+						}
+					}
+
+					if (bMerge)
+					{
+						currentDisplayEvent.count++;
+					}
+					else
+					{
+						currentDisplayEvent = new CsvEventWithCount(ev);
+						filteredEvents.Add(currentDisplayEvent);
+					}
+					lastEvent = ev;
+				}
+			}
+
+			foreach (CsvEventWithCount ev in filteredEvents)
+			{
+				float eventX = ToSvgX(ev.Frame, rect, range);
+				string name = ev.Name;
+
+				if (ev.count > 1)
+				{
+					name += " &#x00D7; " + ev.count;
+				}
+
+				// Space out the events (allow at least 8 pixels between them)
+				if (eventX - LastEventX <= 8.5f)
+				{
+					// Add an arrow to indicate events were spaced out 
+					name = "&#x21b7; " + name;
+					if (ev.count == 1)
+					{
+						name += " (+" + (ev.Frame - lastFrame) + ")";
+					}
+					eventX = LastEventX + 9.0f;
+				}
+				float csvTextX = ToCsvX(eventX + 7, rect, range);
+
+				DrawHorizontalAxisText(svg, name, csvTextX, colour, rect, range, 10, true);
+				LastEventX = eventX;
+				lastFrame = ev.Frame;
+			}
+		}
+
+		void DrawEventLines(SvgFile svg, Theme theme, List<CsvEvent> events, Rect rect, Range range)
+		{
+			float LastEventX = -100000.0f;
+			float LastDisplayedEventX = -100000.0f;
+			int i = 0;
+			foreach (CsvEvent ev in events)
+			{
+				// Only draw events which are in the range
+				if (ev.Frame >= range.MinX && ev.Frame <= range.MaxX)
+				{
+					float eventX = ToSvgX(ev.Frame, rect, range);
+					// Space out the events (allow at least 2 pixels between them)
+					if (eventX - LastDisplayedEventX > 3)
+					{
+						float alpha = (eventX - LastEventX < 1.0f) ? 0.5f : 1.0f;
+						DrawVerticalLine(svg, ev.Frame, theme.MajorGridlineColour, rect, range, alpha, true, true);
+						LastDisplayedEventX = eventX;
+					}
+					LastEventX = eventX;
+				}
+				i++;
+			}
+		}
+
+
+		void DrawEventHighlightRegions(SvgFile svg, CsvStats csvStats, Rect rect, Range range, List<string> highlightEventRegions )
+		{
+			// Draw event regions if requested
+			if (highlightEventRegions != null)
+			{
+				Colour highlightColour = new Colour(0xffffff, 0.1f);
+				float y0 = rect.y;
+				float y1 = rect.y + rect.height;
+				float height = y1 - y0;
+
+				int numPairs = highlightEventRegions.Count / 2;
+				if (highlightEventRegions.Count >= 2 && highlightEventRegions.Count % 2 != 2)
+				{
+					for (int i = 0; i < numPairs; i++)
+					{
+						string startName = highlightEventRegions[i * 2].ToLower().Trim();
+						string endName = highlightEventRegions[i * 2 + 1].ToLower().Trim();
+
+						if (endName == "{NULL}")
+						{
+							endName = null;
+						}
+
+						List<int> startIndices = null;
+						List<int> endIndices = null;
+						csvStats.GetEventFrameIndexDelimiters(startName, endName, out startIndices, out endIndices);
+						for (int j = 0; j < startIndices.Count; j++)
+						{
+							float x0 = ToSvgX(startIndices[j], rect, range);
+							float x1 = ToSvgX(endIndices[j], rect, range);
+							float width = x1 - x0;
+
+							svg.WriteLine("<rect x='" + x0 + "' y='" + y0 + "' width='" + width + "' height='" + height + "' fill=" + highlightColour.SVGString() + "/>");
+						}
+					}
+				}
+			}
+		}
+		void DrawVerticalAxisNotch(SvgFile svg, float sample, Colour colour, Rect rect, Range range, float thickness = 0.25f)
+		{
+			if (sample > range.MaxY) return;
+			if (sample < range.MinY) return;
+
+			float y = ToSvgY(sample, rect, range);
+
+			float x1 = rect.x - 2.5f;
+			float x2 = rect.x;
+
+			svg.WriteLineFast("<line x1='" + x1 + "' y1='" + y + "' x2='" + x2 + "' y2='" + y + "' style='fill:none;stroke-width:" + thickness + "'"
+				+ " stroke=" + colour.SVGString() + "/>");
+		}
+
+		void DrawHorizLine(SvgFile svg, float sample, Colour colour, Rect rect, Range range, bool dashed = false, float thickness = 0.25f, bool dropShadow = false)
+		{
+			if (sample > range.MaxY) return;
+			if (sample < range.MinY) return;
+
+			float y = ToSvgY(sample, rect, range);
+
+			float x1 = rect.x;
+			float x2 = rect.x + rect.width;
+
+			if (dropShadow)
+			{
+				Colour shadowColour = new Colour(0, 0, 0, 0.33f);
+				svg.WriteLineFast("<line x1='" + (x1 - 1) + "' y1='" + (y + 1) + "' x2='" + (x2 - 1) + "' y2='" + (y + 1) + "' style='fill:none;stroke-width:" + thickness * 1.5f + "'"
+					+ " stroke=" + shadowColour.SVGString() + ""
+					+ (dashed ? " stroke-dasharray='4,4' d='M5 10 l215 0' " : "")
+					+ "/>" );
+			}
+
+			svg.WriteLineFast("<line x1='" + x1 + "' y1='" + y + "' x2='" + x2 + "' y2='" + y + "' style='fill:none;stroke-width:" + thickness + "'"
+				+ " stroke=" + colour.SVGString() + ""
+				+ (dashed ? " stroke-dasharray='4,4' d='M5 10 l215 0' " : "")
+				//   + (dropShadow ? "filter='url(#dropShadowFilter)'>" : "" )
+				+ "/>" );
+		}
+
+		void DrawVerticalAxisText(SvgFile svg, string text, float textY, Colour colour, Rect rect, Range range, float size = 10.0f)
+		{
+			float y = ToSvgY(textY, rect, range) + 4;
+			float x = rect.x - 10.0f;
+			svg.WriteLine("<text x='" + x + "' y='" + y + "' fill=" + colour.SVGString() + " font-size='" + size + "' font-family='Helvetica' style='text-anchor: end'>" + text + "</text> >");
+		}
+
+
+		void DrawHorizontalAxisText(SvgFile svg, string text, float textX, Colour colour, Rect rect, Range range, float size = 10.0f, bool top = false)
+		{
+			float scaleX = rect.width / (float)(range.MaxX - range.MinX - 1);
+			float x = ToSvgX(textX, rect, range) - 4;
+			float y = rect.y + rect.height + 10;
+
+			if (top)
+			{
+				y = rect.y + 5;
+			}
+
+			svg.WriteLine("<text x='" + y + "' y='" + (-x) + "' fill=" + colour.SVGString() + " font-size='" + size + "' font-family='Helvetica' transform='rotate(90)' >" + text + "</text>");
+		}
+
+		void DrawLegend(SvgFile svg, Theme theme, List<CsvStats> csvStats, GraphParams graphParams, Rect rect)
+		{
+			List<StatSamples> statSamples = new List<StatSamples>();
+
+			foreach (CsvStats csvStat in csvStats)
+			{
+				foreach (StatSamples stat in csvStat.Stats.Values)
+				{
+					statSamples.Add(stat);
+				}
+			}
+
+			// Stacked stats are already sorted in the right order
+			if (graphParams.forceLegendSort || ((graphParams.showAverages || graphParams.showTotals) && !graphParams.stacked))
+			{
+				statSamples.Sort();
+			}
+
+
+			float x = rect.width + rect.x - 60;
+			if (graphParams.showAverages || graphParams.showTotals)
+			{
+				x -= 30;
+			}
+			float y = 30; // rect.height - 25;
+
+			svg.WriteLine("<g id='LegendPanel<UNIQUE>'>");// transform='translate(5) rotate(45 50 50)'> ");
+
+			// Check if the total is a fraction
+			bool legendValueIsWholeNumber = false;
+			if (graphParams.showTotals)
+			{
+				legendValueIsWholeNumber = true;
+				foreach (StatSamples stat in statSamples)
+				{
+					if (stat.total != Math.Floor(stat.total))
+					{
+						legendValueIsWholeNumber = false;
+						break;
+					}
+				}
+			}
+
+			foreach (StatSamples stat in statSamples)
+			{
+				if (stat.average > graphParams.legendAverageThreshold)
+				{
+					Colour colour = stat.colour;
+					svg.WriteLine("<text x='" + (x - 5) + "' y='" + y + "' fill=" + theme.TextColour.SVGString() + " font-size='10' font-family='Helvetica' style='text-anchor: end' filter='url(#dropShadowFilter)'>" + stat.LegendName + "</text>");
+					svg.WriteLine("<rect x='" + (x + 0) + "' y='" + (y - 8) + "' width='10' height='10' fill=" + colour.SVGString() + " filter='url(#dropShadowFilter)' stroke-width='1' stroke='" + theme.LineColour + "' />");
+
+					if (graphParams.showAverages || graphParams.showTotals)
+					{
+						float legendValue = graphParams.showTotals ? (float)stat.total : stat.average;
+						string formatString = "0.00";
+						if (graphParams.showTotals && legendValueIsWholeNumber)
+						{
+							formatString = "0";
+						}
+						svg.WriteLine("<text x='" + (x + 15) + "' y='" + y + "' fill=" + theme.MediumTextColour.SVGString() + " font-size='10' font-family='Helvetica' style='text-anchor: start' filter='url(#dropShadowFilter)'>" + legendValue.ToString(formatString) + "</text>");
+					}
+
+					y += 16;
+				}
+			}
+
+			if (graphParams.showTotals || graphParams.showAverages)
+			{
+				string legendValueTypeString = graphParams.showTotals ? "(sum)" : "(avg)";
+				svg.WriteLine("<text x='" + (x + 15) + "' y='" + y + "' fill=" + theme.MediumTextColour.SVGString() + " font-size='10' font-family='Helvetica' style='text-anchor: start' filter='url(#dropShadowFilter)'>" + legendValueTypeString + "</text>");
+			}
+
+			svg.WriteLine("</g>");
+		}
+
+
+		void DrawVerticalLine(SvgFile svg, float sampleX, Colour colour, Rect rect, Range range, float thickness = 0.25f, bool dashed = false, bool dropShadow = false, string id = "", bool noTranslate = false)
+		{
+			float x = noTranslate ? sampleX : ToSvgX(sampleX, rect, range);
+			float y1 = rect.y;
+			float y2 = rect.y + rect.height;
+
+			if (dropShadow)
+			{
+				Colour shadowColour = new Colour(0, 0, 0, 0.33f);
+				svg.WriteLine("<line x1='" + (x - 1) + "' y1='" + (y1 + 1) + "' x2='" + (x - 1) + "' y2='" + (y2 + 1) + "' style='fill:none;stroke-width:" + thickness * 1.5f + "' "
+					  + " stroke=" + shadowColour.SVGString() + ""
+					  + (dashed ? " stroke-dasharray='4,4' d='M5 10 l215 0' " : "")
+					  + (id.Length > 0 ? " id='" + id + "_dropShadow'" : "")
+					  + "/>");
+			}
+
+			svg.WriteLine("<line x1='" + x + "' y1='" + y1 + "' x2='" + x + "' y2='" + y2 + "' style='fill:none;stroke-width:" + thickness + "' "
+				+ " stroke=" + colour.SVGString() + ""
+				+ (dashed ? " stroke-dasharray='4,4' d='M5 10 l215 0' " : "")
+				+ (id.Length > 0 ? " id='" + id + "'" : "")
+				+ "/>");
+
+		}
+
+		string GetJSStatName(string statName)
+		{
+			uint hashCode = (uint)statName.GetHashCode();
+			string newString = "_";
+			foreach (char c in statName)
+			{
+				if ((c >= 'A' && c <= 'Z') || (c >= 'a' && c <= 'z'))
+				{
+					newString += c;
+				}
+				else
+				{
+					newString += "_";
+				}
+			}
+			newString += "_" + hashCode.ToString();
+			return newString;
+		}
+
+		class InteractiveStatInfo
+		{
+			public string friendlyName;
+			public string jsVarName;
+			public string jsTextElementId;
+			public string jsGroupElementId;
+			public string graphID;
+			public Colour colour;
+			public StatSamples statSamples;
+			public StatSamples originalStatSamples;
+		};
+
+		void AddInteractiveScripting(SvgFile svg, Theme theme, Rect rect, Range range, List<CsvStats> csvStats, GraphParams graphParams)
+		{
+			bool bSnapToPeaks = graphParams.snapToPeaks && !graphParams.smooth;
+
+			// Todo: separate values, quantized frame pos, pass in X axis name
+			List<InteractiveStatInfo> interactiveStats = new List<InteractiveStatInfo>();
+
+			// Add the frame number stat
+			{
+				InteractiveStatInfo statInfo = new InteractiveStatInfo();
+				statInfo.friendlyName = "Frame";
+				statInfo.jsTextElementId = "csvFrameText<UNIQUE>";
+				statInfo.jsGroupElementId = "csvFrameGroup<UNIQUE>";
+				statInfo.statSamples = null;
+				interactiveStats.Add(statInfo);
+			}
+
+			for (int i = 0; i < csvStats.Count; i++)
+			{
+				foreach (StatSamples stat in csvStats[i].Stats.Values)
+				{
+					InteractiveStatInfo statInfo = new InteractiveStatInfo();
+					statInfo.colour = stat.colour;
+					statInfo.friendlyName = stat.LegendName;
+					statInfo.jsVarName = "statData_" + i + "_" + GetJSStatName(stat.Name) + "<UNIQUE>";
+					statInfo.graphID = "Stat_" + i + "_" + stat.Name + "<UNIQUE>";
+					statInfo.jsTextElementId = "csvStatText__" + i + GetJSStatName(stat.Name) + "<UNIQUE>";
+					statInfo.jsGroupElementId = "csvStatGroup__" + i + GetJSStatName(stat.Name) + "<UNIQUE>";
+					statInfo.originalStatSamples = stat;
+					statInfo.statSamples = new StatSamples(stat, false);
+					interactiveStats.Add(statInfo);
+				}
+			}
+
+			// Record the max value at each index
+			// TODO: strip out data outside the range (needs an offset as well as a multiplier)
+			int multiplier = 1;
+			float numStatsPerPixel = (float)(range.MaxX - range.MinX) / rect.width;
+
+			if (numStatsPerPixel > 1)
+			{
+				multiplier = (int)numStatsPerPixel;
+
+				int maxCount = 0;
+				foreach (InteractiveStatInfo statInfo in interactiveStats)
+				{
+					if (statInfo.originalStatSamples != null)
+					{
+						maxCount = Math.Max(statInfo.originalStatSamples.samples.Count, maxCount);
+					}
+				}
+
+				// Compute max value for each frame
+				List<float> maxValues = new List<float>(maxCount);
+				foreach (InteractiveStatInfo statInfo in interactiveStats)
+				{
+					if (statInfo.originalStatSamples != null)
+					{
+						for (int i = 0; i < statInfo.originalStatSamples.samples.Count; i++)
+						{
+							float value = statInfo.originalStatSamples.samples[i];
+							if (i >= maxValues.Count)
+							{
+								maxValues.Add(value);
+							}
+							else
+							{
+								maxValues[i] = Math.Max(maxValues[i], value);
+							}
+						}
+					}
+				}
+
+				int filteredStatCount = maxValues.Count / multiplier;// (int)(range.MaxX) - (int)(range.MinX);
+
+				// Create the filtered stat array
+				int offset = multiplier / 2;
+				foreach (InteractiveStatInfo statInfo in interactiveStats)
+				{
+					if (statInfo.originalStatSamples == null)
+					{
+						continue;
+					}
+					for (int i = 0; i < filteredStatCount; i++)
+					{
+						int srcStartIndex = Math.Max(i * multiplier - offset, 0);
+						int srcEndIndex = Math.Min(i * multiplier + offset + 1, maxValues.Count);
+
+						int bestIndex = -1;
+						float highestValue = float.MinValue;
+						// Find a good index based on nearby maxValues
+						for (int j = srcStartIndex; j < srcEndIndex; j++)
+						{
+							if (maxValues[j] > highestValue)
+							{
+								highestValue = maxValues[j];
+								bestIndex = j;
+							}
+						}
+						if (bestIndex < statInfo.originalStatSamples.samples.Count)
+						{
+							statInfo.statSamples.samples.Add(statInfo.originalStatSamples.samples[bestIndex]);
+						}
+					}
+				}
+			}
+			else
+			{
+				foreach (InteractiveStatInfo statInfo in interactiveStats)
+				{
+					statInfo.statSamples = statInfo.originalStatSamples;
+				}
+			}
+
+
+			//style='text-anchor: "
+			//svg.WriteLine("<g transform='translate(0,-24)'>");
+			svg.WriteLine("<g id='interactivePanel<UNIQUE>' visibility='hidden'>");
+			DrawVerticalLine(svg, 0, new Colour(255, 255, 255, 0.4f), rect, range, 1.0f, true, true, "", true);
+			svg.WriteLine("<g id='interactivePanelInnerWithRect<UNIQUE>'>");
+			svg.WriteLine("<rect x='0' y='0' width='100' height='100' fill='rgba(0,0,0,0.3)' blend='1' id='interactivePanelRect<UNIQUE>' rx='5' ry='5'/>");
+			svg.WriteLine("<g id='interactivePanelInner<UNIQUE>'>");
+			Rect zeroRect = new Rect(0, 0, 0, 0);
+			foreach (InteractiveStatInfo statInfo in interactiveStats)
+			{
+				svg.WriteLine("<g id='" + statInfo.jsGroupElementId + "' transform='translate(0,0)' visibility='hidden'>");
+				int textXOffset = -5;
+				if (statInfo.colour != null)
+				{
+					svg.WriteLine("<rect x='-16' y='0' width='10' height='10' fill=" + statInfo.colour.SVGString() + " filter='url(#dropShadowFilter)' stroke-width='1' stroke='" + theme.LineColour + "'/>");
+					textXOffset = -20;
+				}
+				DrawText(svg, statInfo.friendlyName, textXOffset, 9, 11, zeroRect, new Colour(255, 255, 255, 1.0f), "end", "Helvetica", "", true);
+				DrawText(svg, " ", 5, 9, 11, zeroRect, new Colour(255, 255, 255, 1.0f), "start", "Helvetica", statInfo.jsTextElementId, true);
+
+				svg.WriteLine("</g>");
+			}
+			svg.WriteLine("</g>"); // interactivePanelInner
+			svg.WriteLine("</g>"); // interactivePanelInnerWithRect
+			svg.WriteLine("</g>"); // interactivepanel
+
+			svg.WriteLine("<script type='application/ecmascript'> <![CDATA[");
+			float oneOverMultiplier = 1.0f / (float)multiplier;
+
+			// Write the data array for each stat
+			foreach (InteractiveStatInfo statInfo in interactiveStats)
+			{
+				if (statInfo.statSamples != null)
+				{
+					svg.Write("var " + statInfo.jsVarName + " = [");
+					foreach (float value in statInfo.statSamples.samples)
+					{
+						svg.WriteFast(value.ToString("0.00") + ",");
+					}
+					svg.WriteLine("]");
+				}
+			}
+
+			svg.WriteLine("function OnLoaded<UNIQUE>(evt)");
+			svg.WriteLine("{");
+			svg.WriteLine("}");
+
+			if (bSnapToPeaks)
+			{
+				// Record the max value at each index
+				List<float> maxValues = new List<float>();
+				foreach (InteractiveStatInfo statInfo in interactiveStats)
+				{
+					if (statInfo.statSamples != null)
+					{
+						for (int i = 0; i < statInfo.statSamples.samples.Count; i++)
+						{
+							float value = statInfo.statSamples.samples[i];
+							if (i >= maxValues.Count)
+							{
+								maxValues.Add(value);
+							}
+							else
+							{
+								maxValues[i] = Math.Max(maxValues[i], value);
+							}
+						}
+					}
+				}
+
+				svg.Write("var maxValues<UNIQUE> = [");
+				foreach (float value in maxValues)
+				{
+					svg.WriteFast(value + ",");
+				}
+				svg.WriteLine("]");
+
+				svg.WriteLine("function FindInterestingFrame<UNIQUE>(x,range)");
+				svg.WriteLine("{");
+				svg.WriteLine("    range *= " + oneOverMultiplier + ";");
+				svg.WriteLine("    x = Math.round( x*" + oneOverMultiplier + " );");
+				svg.WriteLine("    halfRange = Math.round(range/2.0);");
+				svg.WriteLine("    startX = Math.round(Math.max(x-halfRange,0));");
+				svg.WriteLine("    endX = Math.round(Math.min(x+halfRange,maxValues<UNIQUE>.length));");
+				svg.WriteLine("    maxVal = 0.0;");
+				svg.WriteLine("    bestIndex = x;");
+				svg.WriteLine("    for ( i=startX;i<endX;i++) ");
+				svg.WriteLine("    {");
+				svg.WriteLine("        if (maxValues<UNIQUE>[i]>maxVal)");
+				svg.WriteLine("        {");
+				svg.WriteLine("            bestIndex=i;");
+				svg.WriteLine("            maxVal = maxValues<UNIQUE>[i];");
+				svg.WriteLine("        }");
+				svg.WriteLine("     }");
+				svg.WriteLine("     return Math.round( bestIndex *" + multiplier + ");");
+				svg.WriteLine("}");
+				svg.WriteLine("  ");
+			}
+			else
+			{
+				svg.WriteLine("function FindInterestingFrame<UNIQUE>(x,range) { return x; }");
+			}
+
+
+
+
+			svg.WriteLine("function GetGraphX(mouseX)");
+			svg.WriteLine("{");
+			svg.WriteLine("  return (mouseX - " + rect.x + ") * (" + range.MaxX + " - " + range.MinX + ") / " + rect.width + " + " + range.MinX + ";");
+			svg.WriteLine("}");
+			svg.WriteLine("  ");
+
+			svg.WriteLine("function compareSamples(a, b)");
+			svg.WriteLine("{");
+			svg.WriteLine("      if (a.value > b.value)");
+			svg.WriteLine("          return -1;");
+			svg.WriteLine("      if (a.value < b.value)");
+			svg.WriteLine("          return 1;");
+			svg.WriteLine("      return 0;");
+			svg.WriteLine("}");
+
+			svg.WriteLine("function ToSvgX(graphX)");
+			svg.WriteLine("{");
+			svg.WriteLine("    scaleX = " + rect.width / (range.MaxX - range.MinX) + ";");
+			svg.WriteLine("    return " + rect.x + " + (graphX - " + range.MinX + ") * scaleX;");
+			svg.WriteLine("}");
+
+			/*
+            svg.WriteLine("var graphView = { offsetX:0.0, offsetY:0.0, scaleX:1.0, scaleY:1.0 };");
+            svg.WriteLine("function OnGraphAreaWheeled<UNIQUE>(evt)");
+            svg.WriteLine("{");
+            svg.WriteLine("     var graphAreaElement = document.getElementById('graphAreaTransform');");
+            svg.WriteLine("     graphView.scaleX+=0.1;");
+            svg.WriteLine("     ");
+            svg.WriteLine("     graphAreaElement.setAttribute('transform','scale('+graphView.scaleX+','+graphView.scaleY+')')");
+            svg.WriteLine("}");
+*/
+
+			int onePixelFrameCount = (int)(ToCsvXScale(1.0f, rect, range) * 8.0f);
+			float rightEdgeX = ToSvgX(range.MaxX, rect, range);
+
+			svg.WriteLine("function OnGraphAreaClicked<UNIQUE>(evt)");
+			svg.WriteLine("{");
+			svg.WriteLine("  graphX = GetGraphX(evt.offsetX); ");
+			svg.WriteLine("  var interactivePanel = document.getElementById('interactivePanel<UNIQUE>');");
+			svg.WriteLine("  var legendPanel = document.getElementById('LegendPanel<UNIQUE>');");
+			// Snap to an interesting frame (the max value under the pixel)
+
+			svg.WriteLine("  var frameNum = FindInterestingFrame<UNIQUE>( Math.round(graphX+0.5), " + onePixelFrameCount + " );");
+			svg.WriteLine("  if (frameNum >= " + range.MinX + " && frameNum < " + range.MaxX + ")");
+			svg.WriteLine("  {");
+			svg.WriteLine("    var xOffset = 0;");
+			svg.WriteLine("    var lineX = ToSvgX(frameNum);");
+			svg.WriteLine("    var textX = lineX + xOffset;");
+			svg.WriteLine("    var textY = " + rect.y + " - 20");
+
+			svg.WriteLine("    var interactivePanelInner = document.getElementById('interactivePanelInner<UNIQUE>');");
+			svg.WriteLine("    legendPanel.setAttribute('visibility','hidden')");
+			svg.WriteLine("    interactivePanel.setAttribute('visibility','visible')");
+			svg.WriteLine("    interactivePanel.setAttribute('transform','translate('+textX+',0)')");
+			svg.WriteLine("    var dataIndex = Math.round( frameNum * " + oneOverMultiplier + " );");
+
+			// Fill out the sample data array
+			svg.WriteLine("    var samples = [");
+			foreach (InteractiveStatInfo statInfo in interactiveStats)
+			{
+				string groupElementString = "document.getElementById('" + statInfo.jsGroupElementId + "')";
+				string textElementString = "document.getElementById('" + statInfo.jsTextElementId + "')";
+				if (statInfo.jsVarName == null)
+				{
+					svg.Write("            { value: frameNum, name: '" + statInfo.friendlyName + "', colour: 'rgb(0,0,0)'");
+				}
+				else
+				{
+					string valueStr = statInfo.jsVarName + "[dataIndex]";
+					svg.Write("            { value: " + valueStr + ", name: '" + statInfo.friendlyName + "', colour: " + statInfo.colour.SVGString());
+				}
+				svg.WriteLine(", groupElement: " + groupElementString + ", textElement: " + textElementString + " },");
+			}
+			svg.WriteLine("    ];");
+
+			svg.WriteLine("    samples.sort(compareSamples);");
+			svg.WriteLine("    for ( i=0;i<samples.length;i++ )");
+			svg.WriteLine("    {");
+			svg.WriteLine("        var groupElement = samples[i].groupElement;");
+			svg.WriteLine("        if ( samples[i].value > 0 )");
+			svg.WriteLine("        {");
+			svg.WriteLine("            groupElement.setAttribute('visibility','inherit');");
+			svg.WriteLine("            groupElement.setAttribute('transform','translate(0,'+textY+')');");
+			svg.WriteLine("            var textElement = samples[i].textElement;");
+			//            svg.WriteLine("            var textNode = document.createTextNode(interactivePanelInner.getBBox().x);");
+			svg.WriteLine("            var textNode = document.createTextNode(samples[i].value);");
+			svg.WriteLine("            textElement.replaceChild(textNode,textElement.childNodes[0]); ");
+			svg.WriteLine("            textY += 15;");
+			svg.WriteLine("        }");
+			svg.WriteLine("        else");
+			svg.WriteLine("        {");
+			svg.WriteLine("            groupElement.setAttribute('visibility','hidden')");
+			svg.WriteLine("        }");
+			svg.WriteLine("    }");
+
+			svg.WriteLine("    var panelRect = document.getElementById('interactivePanelRect<UNIQUE>');");
+			svg.WriteLine("    var bbox = interactivePanelInner.getBBox();");
+			svg.WriteLine("    panelRect.setAttribute('width',bbox.width+32);");
+			svg.WriteLine("    panelRect.setAttribute('height',textY+8);");
+			svg.WriteLine("    panelRect.setAttribute('x',bbox.x-16);");
+			svg.WriteLine("    panelRect.setAttribute('y',bbox.y-8);");
+
+			float maxX = rect.x + rect.width;
+			svg.WriteLine("    var panelOffset = 0;");
+			svg.WriteLine("    if ( bbox.x + textX < 0 )");
+			svg.WriteLine("    {");
+			svg.WriteLine("        panelOffset = -(bbox.x + textX);");
+			svg.WriteLine("    }");
+			svg.WriteLine("    else if ( bbox.x+bbox.width + textX > " + maxX + " ) ");
+			svg.WriteLine("    {    ");
+			svg.WriteLine("        panelOffset = " + maxX + "-(bbox.x+bbox.width + textX);");
+			svg.WriteLine("    }   ");
+			svg.WriteLine("    var interactivePanelInnerWithRect = document.getElementById('interactivePanelInnerWithRect<UNIQUE>');");
+			svg.WriteLine("    interactivePanelInnerWithRect.setAttribute('transform','translate('+panelOffset+',0)');");
+			svg.WriteLine("  }");
+			svg.WriteLine("  else");
+			svg.WriteLine("  {");
+			svg.WriteLine("    interactivePanel.setAttribute('visibility','hidden')");
+			svg.WriteLine("    legendPanel.setAttribute('visibility','visible')");
+			svg.WriteLine("  }");
+			svg.WriteLine("}");
+			svg.WriteLine("]]> </script>");
+
+			Rect dimensions = new Rect(0, 0, graphParams.width, graphParams.height);
+			DrawGraphArea(svg, theme, new Colour(0, 0, 0, 0.0f), dimensions, false, true);
+		}
+
+		void DrawMetadata(SvgFile svg, CsvMetadata metadata, Rect rect, Colour colour)
+		{
+			float y = rect.height - 15;
+
+			string Platform = metadata.GetValue("platform", "[Unknown platform]");
+			string BuildConfiguration = metadata.GetValue("config", "[Unknown config]");
+			string BuildVersion = metadata.GetValue("buildversion", "[Unknown version]");
+			string SessionId = metadata.GetValue("sessionid", "");
+			string PlaylistId = metadata.GetValue("playlistid", "");
+			// If we have a device profile, write that out instead of the platform
+			Platform = metadata.GetValue("deviceprofile", Platform);
+			BuildVersion = metadata.GetValue("buildversion", BuildVersion);
+
+			string Commandline = metadata.GetValue("commandline", "");
+			svg.WriteLine("<text x='" + 10 + "' y='" + y + "' fill=" + colour.SVGString() + " font-size='9' font-family='Helvetica' >" + BuildConfiguration + " " + Platform + " " + BuildVersion + " " + SessionId + " " + PlaylistId + "</text>");
+			y += 10;
+			svg.WriteLine("<text x='" + 10 + "' y='" + y + "' fill=" + colour.SVGString() + " font-size='9' font-family='Courier New' >" + Commandline + "</text>");
+		}
+
+		bool IsEventShown(string eventString, GraphParams graphParams)
+		{
+			if (graphParams.showEventNames == null)
+			{
+				return false;
+			}
+			if (graphParams.showEventNames.Count == 0)
+			{
+				return false;
+			}
+			if (eventString.Length == 0)
+			{
+				return false;
+			}
+
+			eventString = eventString.ToLower();
+			foreach (string showEventName in graphParams.showEventNames)
+			{
+				string showEventNameLower = showEventName.ToLower();
+				if (showEventNameLower.EndsWith("*"))
+				{
+					int index = showEventNameLower.LastIndexOf('*');
+					string prefix = showEventNameLower.Substring(0, index);
+					if (eventString.StartsWith(prefix))
+					{
+						return true;
+					}
+				}
+				else if (eventString == showEventNameLower)
+				{
+					return true;
+				}
+			}
+			return false;
+		}
+
+		bool IsStatIgnored(string statName, GraphParams graphParams)
+		{
+			statName = statName.ToLower();
+
+			if (graphParams.maxHierarchyDepth != -1)
+			{
+				int Depth = 0;
+				foreach (char c in statName)
+				{
+					if (c == graphParams.hierarchySeparator)
+					{
+						Depth++;
+						if (Depth > graphParams.maxHierarchyDepth)
+						{
+							return true;
+						}
+					}
+				}
+			}
+
+			if (graphParams.ignoreStats.Count == 0)
+			{
+				return false;
+			}
+
+			foreach (string ignoreStat in graphParams.ignoreStats)
+			{
+				string ignoreStatLower = ignoreStat.ToLower();
+				if (ignoreStatLower.EndsWith("*"))
+				{
+					int index = ignoreStatLower.LastIndexOf('*');
+					string prefix = ignoreStatLower.Substring(0, index);
+					if (statName.StartsWith(prefix))
+					{
+						return true;
+					}
+				}
+				else if (statName == ignoreStatLower)
+				{
+					return true;
+				}
+			}
+
+			return false;
+		}
+
+		CsvStats ProcessCsvStats(CsvStats csvStatsIn, GraphParams graphParams)
+		{
+			CsvStats csvStats = new CsvStats(csvStatsIn, graphParams.statNames.ToArray());
+
+			if (graphParams.discardLastFrame)
+			{
+				foreach (StatSamples stat in csvStats.Stats.Values.ToArray())
+				{
+					if (stat.samples.Count > 0)
+					{
+						stat.samples.RemoveAt(stat.samples.Count - 1);
+					}
+				}
+			}
+
+			// Process the minFilterStat
+			float minFilterStatValue = graphParams.minFilterStatValue;
+			StatSamples minFilterStat = null;
+			if (graphParams.minFilterStatName != null && minFilterStatValue > -float.MaxValue)
+			{
+				minFilterStat = csvStats.GetStat(graphParams.minFilterStatName);
+				if (minFilterStat != null)
+				{
+					for (int i = 0; i < minFilterStat.samples.Count; ++i)
+					{
+						if (minFilterStat.samples[i] < minFilterStatValue)
+						{
+							minFilterStat.samples[i] = 0.0f;
+						}
+					}
+					minFilterStat.ComputeAverageAndTotal();
+				}
+				else
+				{
+					// Need a proper stat name for the min filter to work
+					minFilterStatValue = -float.MaxValue;
+				}
+			}
+
+			// Filter out stats which are ignored etc
+			List<StatSamples> FilteredStats = new List<StatSamples>();
+			foreach (StatSamples stat in csvStats.Stats.Values.ToArray())
+			{
+				if (IsStatIgnored(stat.Name, graphParams))
+				{
+					continue;
+				}
+
+				if (minFilterStatValue > -float.MaxValue)
+				{
+					if (stat != minFilterStat)
+					{
+						// Reset other stats when the filter stat entry was below the threshold
+						for (int i = 0; i < stat.samples.Count; i++)
+						{
+							if (minFilterStat.samples[i] == 0.0f)
+							{
+								stat.samples[i] = 0.0f;
+							}
+						}
+					}
+				}
+
+				if (graphParams.filterOutZeros)
+				{
+					float lastNonZero = 0.0f;
+					for (int i = 0; i < stat.samples.Count; i++)
+					{
+						if (stat.samples[i] == 0.0f)
+						{
+							stat.samples[i] = lastNonZero;
+						}
+						else
+						{
+							lastNonZero = stat.samples[i];
+						}
+					}
+					stat.ComputeAverageAndTotal();
+				}
+
+				if (graphParams.statMultiplier != 1.0f)
+				{
+					for (int i = 0; i < stat.samples.Count; i++)
+					{
+						stat.samples[i] *= graphParams.statMultiplier;
+					}
+				}
+
+				// Filter out stats where the average below averageThreshold
+				if (graphParams.averageThreshold > -float.MaxValue)
+				{
+					if (stat.average < graphParams.averageThreshold)
+					{
+						continue;
+					}
+				}
+
+				// Filter out stats below the threshold
+				if (graphParams.threshold > -float.MaxValue)
+				{
+					bool aboveThreshold = false;
+					foreach (float val in stat.samples)
+					{
+						if (val > graphParams.threshold)
+						{
+							aboveThreshold = true;
+							break;
+						}
+					}
+
+					if (!aboveThreshold)
+					{
+						continue;
+					}
+				}
+
+				// LLM specific, spaces in LLM seem to be $32$.
+				// This is a temp fix until LLM outputs without $32$.
+				stat.Name = stat.Name.Replace("$32$", " ");
+				// If we get here, the stat wasn't filtered
+				FilteredStats.Add(stat);
+			}
+
+			// Have any stats actually been filtered? If so, replace the list
+			if (FilteredStats.Count < csvStats.Stats.Count)
+			{
+				csvStats.Stats.Clear();
+				foreach (StatSamples stat in FilteredStats)
+				{
+					csvStats.AddStat(stat);
+				}
+			}
+
+			// Filter out events
+			List<CsvEvent> FilteredEvents = new List<CsvEvent>();
+			foreach (CsvEvent ev in csvStats.Events)
+			{
+				if (IsEventShown(ev.Name, graphParams))
+				{
+					FilteredEvents.Add(ev);
+				}
+			}
+			csvStats.Events = FilteredEvents;
+
+			return csvStats;
+		}
+
+		void RecomputeStatAveragesForRange(CsvStats stats, Range range)
+		{
+			foreach (StatSamples stat in stats.Stats.Values.ToArray())
+			{
+				int rangeStart = Math.Max(0, (int)range.MinX);
+				int rangeEnd = Math.Min(stat.samples.Count, (int)range.MaxX);
+
+				stat.ComputeAverageAndTotal(rangeStart, rangeEnd);
+			}
+		}
+	}
+
+
+}
\ No newline at end of file
diff --git a/Engine/Source/Programs/CSVTools/CsvToSvgLib/CsvToSvgLib.csproj b/Engine/Source/Programs/CSVTools/CsvToSvgLib/CsvToSvgLib.csproj
new file mode 100644
--- /dev/null
+++ b/Engine/Source/Programs/CSVTools/CsvToSvgLib/CsvToSvgLib.csproj
@@ -0,0 +1,61 @@
+﻿<?xml version="1.0" encoding="utf-8"?>
+<Project ToolsVersion="12.0" DefaultTargets="Build" xmlns="http://schemas.microsoft.com/developer/msbuild/2003">
+  <Import Project="$(MSBuildExtensionsPath)\$(MSBuildToolsVersion)\Microsoft.Common.props" Condition="Exists('$(MSBuildExtensionsPath)\$(MSBuildToolsVersion)\Microsoft.Common.props')" />
+  <PropertyGroup>
+    <Configuration Condition=" '$(Configuration)' == '' ">Debug</Configuration>
+    <Platform Condition=" '$(Platform)' == '' ">AnyCPU</Platform>
+    <ProjectGuid>{75C595CC-83BA-4FFD-ABCD-17B2343D1061}</ProjectGuid>
+    <OutputType>Library</OutputType>
+    <AppDesignerFolder>Properties</AppDesignerFolder>
+    <RootNamespace>CsvStats</RootNamespace>
+    <AssemblyName>CsvToSvgLib</AssemblyName>
+    <TargetFrameworkVersion>v4.5</TargetFrameworkVersion>
+    <FileAlignment>512</FileAlignment>
+  </PropertyGroup>
+  <PropertyGroup Condition=" '$(Configuration)|$(Platform)' == 'Debug|AnyCPU' ">
+    <DebugSymbols>true</DebugSymbols>
+    <DebugType>full</DebugType>
+    <Optimize>false</Optimize>
+    <OutputPath>bin\Debug\</OutputPath>
+    <DefineConstants>DEBUG;TRACE</DefineConstants>
+    <ErrorReport>prompt</ErrorReport>
+    <WarningLevel>4</WarningLevel>
+    <PlatformTarget>x64</PlatformTarget>
+  </PropertyGroup>
+  <PropertyGroup Condition=" '$(Configuration)|$(Platform)' == 'Release|AnyCPU' ">
+    <DebugType>pdbonly</DebugType>
+    <Optimize>true</Optimize>
+    <OutputPath>..\..\..\..\Binaries\DotNET\CsvTools\</OutputPath>
+    <DefineConstants>TRACE</DefineConstants>
+    <ErrorReport>prompt</ErrorReport>
+    <WarningLevel>4</WarningLevel>
+    <PlatformTarget>x64</PlatformTarget>
+  </PropertyGroup>
+  <ItemGroup>
+    <Reference Include="System" />
+    <Reference Include="System.Core" />
+    <Reference Include="System.Xml.Linq" />
+    <Reference Include="System.Data.DataSetExtensions" />
+    <Reference Include="Microsoft.CSharp" />
+    <Reference Include="System.Data" />
+    <Reference Include="System.Xml" />
+  </ItemGroup>
+  <ItemGroup>
+    <Compile Include="CsvToSvgLib.cs" />
+    <Compile Include="Properties\AssemblyInfo.cs" />
+  </ItemGroup>
+  <ItemGroup>
+    <ProjectReference Include="..\CsvStats\CSVStats.csproj">
+      <Project>{75c595cc-83ba-4ffd-a772-17b2343d1061}</Project>
+      <Name>CSVStats</Name>
+    </ProjectReference>
+  </ItemGroup>
+  <Import Project="$(MSBuildToolsPath)\Microsoft.CSharp.targets" />
+  <!-- To modify your build process, add your task inside one of the targets below and uncomment it. 
+       Other similar extension points exist, see Microsoft.Common.targets.
+  <Target Name="BeforeBuild">
+  </Target>
+  <Target Name="AfterBuild">
+  </Target>
+  -->
+</Project>
\ No newline at end of file
diff --git a/Engine/Source/Programs/CSVTools/CsvToSvgLib/Properties/AssemblyInfo.cs b/Engine/Source/Programs/CSVTools/CsvToSvgLib/Properties/AssemblyInfo.cs
new file mode 100644
--- /dev/null
+++ b/Engine/Source/Programs/CSVTools/CsvToSvgLib/Properties/AssemblyInfo.cs
@@ -0,0 +1,39 @@
+// Copyright (C) Microsoft. All rights reserved.
+// Copyright Epic Games, Inc. All Rights Reserved.
+
+using System.Reflection;
+using System.Runtime.CompilerServices;
+using System.Runtime.InteropServices;
+
+// General Information about an assembly is controlled through the following 
+// set of attributes. Change these attribute values to modify the information
+// associated with an assembly.
+[assembly: AssemblyTitle("CsvToSvgLib")]
+[assembly: AssemblyDescription("")]
+[assembly: AssemblyConfiguration("")]
+[assembly: AssemblyCompany("")]
+[assembly: AssemblyProduct("CsvToSvgLib")]
+[assembly: AssemblyCopyright("Copyright (C) Microsoft. All rights reserved.\nCopyright Epic Games, Inc. All Rights Reserved.")]
+[assembly: AssemblyTrademark("")]
+[assembly: AssemblyCulture("")]
+
+// Setting ComVisible to false makes the types in this assembly not visible 
+// to COM components.  If you need to access a type in this assembly from 
+// COM, set the ComVisible attribute to true on that type.
+[assembly: ComVisible(false)]
+
+// The following GUID is for the ID of the typelib if this project is exposed to COM
+[assembly: Guid("54811cbd-f2b0-427f-ABCD-2f93318baab1")]
+
+// Version information for an assembly consists of the following four values:
+//
+//      Major Version
+//      Minor Version 
+//      Build Number
+//      Revision
+//
+// You can specify all the values or you can default the Build and Revision Numbers 
+// by using the '*' as shown below:
+// [assembly: AssemblyVersion("1.0.*")]
+[assembly: AssemblyVersion("1.0.0.0")]
+[assembly: AssemblyFileVersion("1.0.0.0")]
