PerfReportTool 4.72 - use CsvToSvgLib async tasks instead of processes by default (override to legacy behavior with -csvToSvgProcesses). This reduced bulk report generation from 12s to 7.7s in tests
- fix a bug in custom commandline parsing where last character would be skipped
- Update binaries for PerfReportTool, CsvToSvg, CsvToSvgLib, CsvStats

diff --git a/Engine/Source/Programs/CSVTools/CsvStats/CommandLineTool.cs b/Engine/Source/Programs/CSVTools/CsvStats/CommandLineTool.cs
--- a/Engine/Source/Programs/CSVTools/CsvStats/CommandLineTool.cs
+++ b/Engine/Source/Programs/CSVTools/CsvStats/CommandLineTool.cs
@@ -173,11 +173,10 @@ namespace CSVStats
 			{
 				return CommandLineArgs[lowerKey];
 			}
-			else if (mandatory)
+			if (mandatory)
 			{
 				Console.WriteLine("Missing parameter " + key);
 			}
-
 			return "";
 		}
 
@@ -252,6 +251,20 @@ namespace CSVStats
 			return commandLine.GetArg(key, mandatory);
         }
 
+		protected List<string> GetListArg(string key, char separator=';', bool convertToLowercase=false, bool mandatory=false)
+		{
+			string listStr=commandLine.GetArg(key, mandatory);
+			if (listStr=="")
+			{
+				return new List<string>();
+			}
+			if (convertToLowercase)
+			{
+				listStr = listStr.ToLower();
+			}
+			return listStr.Split(separator).ToList();
+		}
+
 		protected void WriteLine(String message, params object[] args)
         {
             String formatted = String.Format(message, args);
@@ -301,4 +314,55 @@ namespace CSVStats
 			commandLine = new CommandLine(args);
 		}
 	}
+
+
+	public class PerfLog
+	{
+		public PerfLog(bool inLoggingEnabled, string inLogPrefix=null)
+		{
+			stopWatch = Stopwatch.StartNew();
+			lastTimeElapsed = 0.0;
+			loggingEnabled = inLoggingEnabled;
+			if (inLogPrefix != null)
+			{
+				logPrefix = inLogPrefix + " - ";
+			}
+		}
+
+		public double LogTiming(string description, bool newLine = false)
+		{
+			double elapsed = stopWatch.Elapsed.TotalSeconds - lastTimeElapsed;
+			if (loggingEnabled)
+			{
+				Console.WriteLine("[PerfLog] " + logPrefix + String.Format("{0,-25} : {1,-10}", description, (elapsed * 1000.0).ToString("0.0") + "ms"), 70);
+				if (newLine)
+				{
+					Console.WriteLine();
+				}
+			}
+			lastTimeElapsed = stopWatch.Elapsed.TotalSeconds;
+			return elapsed;
+		}
+
+		public double LogTotalTiming(bool seconds=true)
+		{
+			double elapsed = stopWatch.Elapsed.TotalSeconds;
+			if (loggingEnabled)
+			{
+				if (seconds)
+				{
+					Console.WriteLine("[PerfLog] "+logPrefix+"TOTAL: " + elapsed.ToString("0.0") + "s\n");
+				}
+				else
+				{
+					Console.WriteLine("[PerfLog] " + logPrefix + "TOTAL: " + (elapsed*1000.0).ToString("0.0") + "ms\n");
+				}
+			}
+			return elapsed;
+		}
+		Stopwatch stopWatch;
+		double lastTimeElapsed;
+		bool loggingEnabled;
+		string logPrefix = "";
+	}
 }
diff --git a/Engine/Source/Programs/CSVTools/PerfReportTool/App.config b/Engine/Source/Programs/CSVTools/PerfReportTool/App.config
--- a/Engine/Source/Programs/CSVTools/PerfReportTool/App.config
+++ b/Engine/Source/Programs/CSVTools/PerfReportTool/App.config
@@ -1,6 +1,9 @@
-﻿<?xml version="1.0" encoding="utf-8" ?>
+<?xml version="1.0" encoding="utf-8" ?>
 <configuration>
     <startup> 
         <supportedRuntime version="v4.0" sku=".NETFramework,Version=v4.5" />
     </startup>
+  <runtime>
+    <gcConcurrent enabled="false" />
+  </runtime>
 </configuration>
\ No newline at end of file
diff --git a/Engine/Source/Programs/CSVTools/PerfReportTool/PerfReportTool.cs b/Engine/Source/Programs/CSVTools/PerfReportTool/PerfReportTool.cs
--- a/Engine/Source/Programs/CSVTools/PerfReportTool/PerfReportTool.cs
+++ b/Engine/Source/Programs/CSVTools/PerfReportTool/PerfReportTool.cs
@@ -11,15 +11,18 @@ using System.Diagnostics;
 using CSVStats;
 using System.Collections;
 using System.Security.Cryptography;
+using System.Threading.Tasks;
+using System.Threading;
 
 using PerfSummaries;
 using System.Globalization;
+using CSVTools;
 
 namespace PerfReportTool
 {
     class Version
     {
-        private static string VersionString = "4.71";
+        private static string VersionString = "4.72";
 
         public static string Get() { return VersionString; }
     };
@@ -99,9 +102,12 @@ namespace PerfReportTool
 			"\n" +
 			"Performance args:\n" +
 			"       -perfLog : output performance logging information\n" +
-			"       -noBatchedGraphs : disable batched/multithreaded graph generation (default is enabled)\n" +
 			"       -graphThreads : use with -batchedGraphs to control the number of threads per CsvToSVG instance \n"+
 			"                       (default: PC core count/2)\n" +
+			"       -csvToSvgSequential : Run CsvToSvg sequentially\n" +
+			"Deprecated performance args:\n"+
+			"       -csvToSvgProcesses : Use separate processes for csvToSVG instead of threads (slower)\n" +
+			"       -noBatchedGraphs : disable batched/multithreaded graph generation (use with -csvToSvgProcesses. Default is enabled)\n" +
 			"\n" +
 			"Options to truncate or filter source data:\n" +
 			"Warning: these options disable Summary Table caching\n" +
@@ -694,45 +700,6 @@ namespace PerfReportTool
 
 		static Dictionary<string, bool> UniqueHTMLFilemameLookup = new Dictionary<string, bool>();
 
-        class PerfLog
-        {
-            public PerfLog(bool inLoggingEnabled)
-            {
-                stopWatch = Stopwatch.StartNew();
-                lastTimeElapsed = 0.0;
-                loggingEnabled = inLoggingEnabled;
-            }
-
-            public double LogTiming(string description, bool newLine=false)
-            {
-                double elapsed = stopWatch.Elapsed.TotalSeconds-lastTimeElapsed;
-                if (loggingEnabled)
-                {
-                    Console.WriteLine("[PerfLog] "+
-                        String.Format("{0,-25} : {1,-10}", description , (elapsed*1000.0).ToString("0.0") + "ms"), 70);
-                    if ( newLine )
-                    {
-                        Console.WriteLine();
-                    }
-                }
-                lastTimeElapsed = stopWatch.Elapsed.TotalSeconds;
-                return elapsed;
-            }
-
-            public double LogTotalTiming()
-            {
-                double elapsed = stopWatch.Elapsed.TotalSeconds;
-                if (loggingEnabled)
-                {
-                    Console.WriteLine("[PerfLog] TOTAL: " + elapsed.ToString("0.0") + "s\n");
-                }
-                return elapsed;
-            }
-            Stopwatch stopWatch;
-            double lastTimeElapsed;
-            bool loggingEnabled;
-        }
-
 		void WriteCleanCsv(CachedCsvFile csvFile, string outCsvFilename, ReportTypeInfo reportTypeInfo)
 		{
 			if ( File.Exists(outCsvFilename) )
@@ -818,16 +785,31 @@ namespace PerfReportTool
 				htmlFilename += ".html";
             }
 
-            float thickness = 1.0f;
+			bool bCsvToSvgProcesses = GetBoolArg("csvToSvgProcesses");
+			bool bCsvToSvgMultiThreaded = !GetBoolArg("csvToSvgSequential");
+
+			int minWorker, minIOC;
+			// Get the current settings.
+			ThreadPool.GetMinThreads(out minWorker, out minIOC);
+
+
+			float thickness = 1.0f;
 			List<string> csvToSvgCommandlines = new List<string>();
 			List<string> svgFilenames = new List<string>();
 			string responseFilename = null;
 			List<Process> csvToSvgProcesses = new List<Process>();
+			List<Task> csvToSvgTasks = new List<Task>();
             if (writeDetailedReport)
             {
+				GraphGenerator graphGenerator = null;
+				if (!bCsvToSvgProcesses)
+				{
+					graphGenerator = new GraphGenerator(csvFile.GetFinalCsv(), csvFile.filename);
+				}
+
 				// Generate all the graphs asyncronously
-                foreach (ReportGraph graph in reportTypeInfo.graphs)
-                {
+				foreach (ReportGraph graph in reportTypeInfo.graphs)
+				{
 					string svgFilename = String.Empty;
 					if (graph.isExternal && !GetBoolArg("externalGraphs"))
 					{
@@ -836,32 +818,49 @@ namespace PerfReportTool
 					}
 					bool bFoundStat = false;
 					foreach (string statString in graph.settings.statString.value.Split(','))
-                    {
-                        List<StatSamples> matchingStats = csvFile.dummyCsvStats.GetStatsMatchingString(statString);
-                        if (matchingStats.Count > 0)
-                        {
-                            bFoundStat = true;
-                            break;
-                        }
-                    }
+					{
+						List<StatSamples> matchingStats = csvFile.dummyCsvStats.GetStatsMatchingString(statString);
+						if (matchingStats.Count > 0)
+						{
+							bFoundStat = true;
+							break;
+						}
+					}
 					if (bFoundStat)
 					{
 						svgFilename = GetTempFilename(csvFile.filename) + ".svg";
-						string args = GetCsvToSvgArgs(csvFile.filename, svgFilename, graph, thickness, minX, maxX, false, svgFilenames.Count);
-						if (bBatchedGraphs)
+						if (graphGenerator != null)
 						{
-							csvToSvgCommandlines.Add(args);
+							GraphParams graphParams=GetCsvToSvgGraphParams(csvFile.filename, graph, thickness, minX, maxX, false, svgFilenames.Count);
+							string args = GetCsvToSvgArgs(csvFile.filename, svgFilename, graph, thickness, minX, maxX, false, svgFilenames.Count);
+
+							if (bCsvToSvgMultiThreaded)
+							{
+								csvToSvgTasks.Add( graphGenerator.MakeGraphAsync(graphParams, svgFilename, true, false) );
+							}
+							else
+							{
+								graphGenerator.MakeGraph(graphParams, svgFilename, true, false);
+							}							
 						}
 						else
 						{
-							Process csvToSvgProcess = LaunchCsvToSvgAsync(args);
-							csvToSvgProcesses.Add(csvToSvgProcess);
+							string args = GetCsvToSvgArgs(csvFile.filename, svgFilename, graph, thickness, minX, maxX, false, svgFilenames.Count);
+							if (bBatchedGraphs)
+							{
+								csvToSvgCommandlines.Add(args);
+							}
+							else
+							{
+								Process csvToSvgProcess = LaunchCsvToSvgAsync(args);
+								csvToSvgProcesses.Add(csvToSvgProcess);
+							}
 						}
 					}
 					svgFilenames.Add(svgFilename);
-                }
+				}
 
-				if (bBatchedGraphs)
+				if (bCsvToSvgProcesses && bBatchedGraphs)
 				{
 					// Save the response file
 					responseFilename = GetTempFilename(csvFile.filename) + "_response.txt";
@@ -879,12 +878,23 @@ namespace PerfReportTool
 				bStripStatsByEvents = false;
 			}
 
-			// Read the full csv while we wait for the graph processes to complete
+			if (writeDetailedReport && csvToSvgTasks.Count > 0)
+			{
+				// wait on the graph tasks to complete
+				// Note that we have to do this before we can call ProcessCSV, since this modifies the CsvStats object the graph tasks are reading
+				foreach (Task task in csvToSvgTasks)
+				{
+					task.Wait();
+				}
+				perfLog.LogTiming("    WaitForAsyncGraphs");
+			}
+
+			// Read the full csv while we wait for the graph processes to complete (this is only safe for CsvToSVG processes, not task threads)
 			int numFramesStripped;
 			CsvStats unstrippedCsvStats;
 			CsvStats csvStats=ProcessCsv(csvFile, out numFramesStripped, out unstrippedCsvStats, minX, maxX, perfLog, bStripStatsByEvents);
 
-            if ( writeDetailedReport )
+            if ( writeDetailedReport && csvToSvgProcesses.Count > 0)
             { 
                 // wait on the graph processes to complete
                 foreach (Process process in csvToSvgProcesses)
@@ -1378,7 +1388,7 @@ namespace PerfReportTool
 			double compression = graphSettings.compression.value;
 			int width = graphSettings.width.value;
 			int height = graphSettings.height.value;
-			string additionalArgs = graphSettings.additionalArgs.value;
+			string additionalArgs = "";
 			bool stacked = graphSettings.stacked.value;
 			bool showAverages = graphSettings.showAverages.value;
 			bool filterOutZeros = graphSettings.filterOutZeros.value;
@@ -1439,15 +1449,15 @@ namespace PerfReportTool
 			string args =
 				" -csvs \"" + csvFilename + "\"" +
 				" -title \"" + title + "\"" +
-				" -o \"" + svgFilename +"\"" +
+				" -o \"" + svgFilename + "\"" +
 				" -stats " + statString +
 				" -width " + (width * scaleby).ToString() +
 				" -height " + (height * scaleby).ToString() +
-				OptionalHelper.GetDoubleSetting(graph.budget, " -budget ") + 
+				OptionalHelper.GetDoubleSetting(graph.budget, " -budget ") +
 				" -maxy " + maxy.ToString() +
 				" -uniqueID Graph_" + graphIndex.ToString() +
 				" -lineDecimalPlaces " + lineDecimalPlaces.ToString() +
-				" -nocommandlineEmbed "+
+				" -nocommandlineEmbed " +
 
 				((statMultiplier != 1.0) ? " -statMultiplier " + statMultiplier.ToString("0.0000000000000000000000") : "") +
 				(hideEventNames ? " -hideeventNames 1" : "") +
@@ -1478,6 +1488,148 @@ namespace PerfReportTool
 			return args;
 		}
 
+
+		GraphParams GetCsvToSvgGraphParams(string csvFilename, ReportGraph graph, double thicknessMultiplier, int minx, int maxx, bool multipleCSVs, int graphIndex, float scaleby = 1.0f)
+		{
+			GraphParams graphParams = new GraphParams();
+			graphParams.title = graph.title;
+
+			GraphSettings graphSettings = graph.settings;
+			graphParams.statNames.AddRange( graphSettings.statString.value.Split(',') );
+			graphParams.lineThickness = (float)(graphSettings.thickness.value * thicknessMultiplier);
+			graphParams.smooth = graphSettings.smooth.value && !GetBoolArg("nosmooth");
+			if (graphParams.smooth)
+			{
+				if (graphSettings.smoothKernelPercent.isSet && graphSettings.smoothKernelPercent.value > 0)
+				{
+					graphParams.smoothKernelPercent = (float)graphSettings.smoothKernelPercent.value;
+				}
+				if (graphSettings.smoothKernelSize.isSet && graphSettings.smoothKernelSize.value > 0)
+				{
+					graphParams.smoothKernelSize = (int)(graphSettings.smoothKernelSize.value);
+				}
+			}
+
+			if (graphSettings.compression.isSet)
+			{
+				graphParams.compression = (float)graphSettings.compression.value;
+			}
+			graphParams.width = (int)(graphSettings.width.value * scaleby);
+			graphParams.height = (int)(graphSettings.height.value * scaleby);
+			if ( graphSettings.stacked.isSet )
+			{
+				graphParams.stacked = graphSettings.stacked.value;
+				if (graphParams.stacked)
+				{
+					graphParams.forceLegendSort = true;
+					if (graphSettings.mainStat.isSet)
+					{
+						graphParams.stackTotalStat = graphSettings.mainStat.value;
+					}
+				}
+			}
+			if (graphSettings.showAverages.isSet)
+			{
+				graphParams.showAverages = graphSettings.showAverages.value;
+			}
+			if (graphSettings.filterOutZeros.isSet)
+			{
+				graphParams.filterOutZeros = graphSettings.filterOutZeros.value;
+			}
+			graphParams.snapToPeaks = false;
+			if (graphSettings.snapToPeaks.isSet)
+			{
+				graphParams.snapToPeaks = graphSettings.snapToPeaks.value;
+			}
+
+			graphParams.lineDecimalPlaces = graphSettings.lineDecimalPlaces.isSet ? graphSettings.lineDecimalPlaces.value : 1;
+			if (graphSettings.maxHierarchyDepth.isSet)
+			{
+				graphParams.maxHierarchyDepth = graphSettings.maxHierarchyDepth.value;
+			}
+			if (graphSettings.hideStatPrefix.isSet && graphSettings.hideStatPrefix.value.Length > 0)
+			{
+				graphParams.hideStatPrefixes.AddRange(graphSettings.hideStatPrefix.value.Split(' ',';'));
+			}
+			if (multipleCSVs)
+			{
+				graphParams.showEventNames.Add("CSV:*");
+				graphParams.showEventNameText = false;
+			}
+			else
+			{
+				if (graphSettings.showEvents.isSet && graphSettings.showEvents.value.Length > 0)
+				{
+					graphParams.showEventNames.AddRange(graphSettings.showEvents.value.Split(' ', ';'));
+				}
+			}
+			if (graphSettings.statMultiplier.isSet)
+			{
+				graphParams.statMultiplier = (float)graphSettings.statMultiplier.value;
+			}
+			graphParams.interactive = true;
+
+			if (!GetBoolArg("noStripEvents"))
+			{
+				List<CsvEventStripInfo> eventsToStrip = reportXML.GetCsvEventsToStrip();
+				if (eventsToStrip != null)
+				{
+					for (int i = 0; i < eventsToStrip.Count; i++)
+					{
+						graphParams.highlightEventRegions.Add((eventsToStrip[i].beginName==null) ? "" : eventsToStrip[i].beginName);
+						graphParams.highlightEventRegions.Add((eventsToStrip[i].endName == null) ? "{NULL}" : eventsToStrip[i].endName);
+					}
+				}
+			}
+
+			if ( graph.minFilterStatValue.isSet )
+			{
+				graphParams.minFilterStatValue = (float)graph.minFilterStatValue.value;
+			}
+			if (graphSettings.minFilterStatName.isSet)
+			{
+				graphParams.minFilterStatName = graphSettings.minFilterStatName.value;
+			}
+			if (graph.budget.isSet)
+			{
+				graphParams.budget = (float)graph.budget.value;
+			}
+			graphParams.uniqueId = "Graph_" + graphIndex.ToString();
+
+			if (minx>0)
+			{
+				graphParams.minX = minx;
+			}
+			if (maxx != Int32.MaxValue)
+			{
+				graphParams.maxX = maxx;
+			}
+			if ( graphSettings.miny.isSet )
+			{
+				graphParams.minY = (float)graphSettings.miny.value;
+			}
+			graphParams.maxY = GetFloatArg("maxy", (float)graphSettings.maxy.value);
+
+			if (graphSettings.threshold.isSet)
+			{
+				graphParams.threshold = (float)graphSettings.threshold.value;
+			}
+			if (graphSettings.averageThreshold.isSet)
+			{
+				graphParams.averageThreshold = (float)graphSettings.averageThreshold.value;
+			}
+			if (graphSettings.legendAverageThreshold.isSet)
+			{
+				graphParams.legendAverageThreshold = (float)graphSettings.legendAverageThreshold.value;
+			}
+			if (graphSettings.ignoreStats.isSet && graphSettings.ignoreStats.value.Length > 0)
+			{
+				graphParams.ignoreStats.AddRange(graphSettings.ignoreStats.value.Split(' ', ';'));
+			}
+			return graphParams;
+		}
+
+
 		Process LaunchCsvToSvgAsync(string args)
 		{
 			string csvToolPath = GetBaseDirectory() + "/CSVToSVG.exe";
diff --git a/Engine/Source/Programs/CSVTools/PerfReportTool/ReportGraph.cs b/Engine/Source/Programs/CSVTools/PerfReportTool/ReportGraph.cs
--- a/Engine/Source/Programs/CSVTools/PerfReportTool/ReportGraph.cs
+++ b/Engine/Source/Programs/CSVTools/PerfReportTool/ReportGraph.cs
@@ -196,7 +196,7 @@ namespace PerfReportTool
 			ignoreStats = new OptionalString(element, "ignoreStats");
 
 			statString = new OptionalString(element, "statString", true);
-			additionalArgs = new OptionalString(element, "additionalArgs", true);
+			//additionalArgs = new OptionalString(element, "additionalArgs", true);
 			statMultiplier = new OptionalDouble(element, "statMultiplier");
 			legendAverageThreshold = new OptionalDouble(element, "legendAverageThreshold");
 			snapToPeaks = new OptionalBool(element, "snapToPeaks");
@@ -218,7 +218,7 @@ namespace PerfReportTool
 			compression.InheritFrom(baseSettings.compression);
 			width.InheritFrom(baseSettings.width);
 			height.InheritFrom(baseSettings.height);
-			additionalArgs.InheritFrom(baseSettings.additionalArgs);
+			//additionalArgs.InheritFrom(baseSettings.additionalArgs);
 			stacked.InheritFrom(baseSettings.stacked);
 			showAverages.InheritFrom(baseSettings.showAverages);
 			filterOutZeros.InheritFrom(baseSettings.filterOutZeros);
@@ -248,7 +248,7 @@ namespace PerfReportTool
 		public OptionalDouble compression;
 		public OptionalInt width;
 		public OptionalInt height;
-		public OptionalString additionalArgs;
+		//public OptionalString additionalArgs;
 		public OptionalBool stacked;
 		public OptionalBool showAverages;
 		public OptionalBool filterOutZeros;
