Add support for fixed point graphs (enabled by default)
Fixed point graphs are smaller because they don't include the decimal point. They're also faster to write. We need to apply a scale in order to have subpixel accuracy. The scale is specified by -fixedpointPrecisionScale <N> (default is 2)
- This improves perf report generation speed by ~15% and report size by 20%.

diff --git a/Engine/Source/Programs/CSVTools/CSVToSVG/CsvToSvgProgram.cs b/Engine/Source/Programs/CSVTools/CSVToSVG/CsvToSvgProgram.cs
--- a/Engine/Source/Programs/CSVTools/CSVToSVG/CsvToSvgProgram.cs
+++ b/Engine/Source/Programs/CSVTools/CSVToSVG/CsvToSvgProgram.cs
@@ -59,6 +59,8 @@ namespace CSVTools
 			"       -compression <pixel error value>\n" +
 			"       -discardLastFrame <1|0> (default 1)\n" +
 			"       -filterOutZeros\n" +
+			"       -fixedPointGraphs 1|0 (default 1 - smaller graphs but no subpixel accuracy)\n"+
+			"       -fixedPointPrecisionScale <1..N> - scale for fixed point graph rendering (>1 gives subpixel accuracy)"+
 			"       -graphOnly\n" +
 			"       -hideEventNames <1|0>\n" +
 			"       -hideStatPrefix <list>\n" +
@@ -98,7 +100,7 @@ namespace CSVTools
             "       -percentile99 \n" +
 			"       -uniqueId <string> : unique ID for JS (needed if this is getting embedded in HTML alongside other graphs)\n" +
 			"       -nocommandlineEmbed : don't embed the commandline in the SVG" +
-			"       -lineDecimalPlaces <N> (default 3)" +
+			"       -lineDecimalPlaces <N> (only effective with with -fixedPointGraphs 0. Default=3)" +
 			"       -frameOffset <N> : offset used for frame display name (default 0)" +
             "";
 
@@ -332,10 +334,15 @@ namespace CSVTools
 			// Other flags
 			graphParams.showMetadata = GetArg("nometadata") != "1";
 			graphParams.graphOnly = GetArg("graphOnly") == "1";
-			graphParams.compression = GetFloatArg("compression", graphParams.compression);
 			graphParams.interactive = GetArg("interactive") == "1";
 			graphParams.snapToPeaks = !GetBoolArg("noSnap");
 
+			// Compression
+			graphParams.compression = GetFloatArg("compression", graphParams.compression);
+			graphParams.bFixedPointGraphs = GetIntArg("fixedPointGraphs", graphParams.bFixedPointGraphs?1:0) == 1;
+			graphParams.fixedPointPrecisionScale = GetFloatArg("fixedPointPrecisionScale", graphParams.fixedPointPrecisionScale);
+			graphParams.lineDecimalPlaces=GetIntArg("lineDecimalPlaces", graphParams.lineDecimalPlaces);
+
 			// Sometimes the last frame is garbage. We might want to remove it
 			graphParams.discardLastFrame = (GetIntArg("discardLastFrame", graphParams.discardLastFrame?1:0) == 1);
 
@@ -397,7 +404,6 @@ namespace CSVTools
 			}
 
 			// Misc params (rarely used)
-			graphParams.lineDecimalPlaces=GetIntArg("lineDecimalPlaces", graphParams.lineDecimalPlaces);
 			graphParams.frameOffset = GetIntArg("frameOffset", graphParams.frameOffset);
 			graphParams.statMultiplier = GetFloatArg("statMultiplier", graphParams.statMultiplier);
 
diff --git a/Engine/Source/Programs/CSVTools/CsvToSvgLib/CsvToSvgLib.cs b/Engine/Source/Programs/CSVTools/CsvToSvgLib/CsvToSvgLib.cs
--- a/Engine/Source/Programs/CSVTools/CsvToSvgLib/CsvToSvgLib.cs
+++ b/Engine/Source/Programs/CSVTools/CsvToSvgLib/CsvToSvgLib.cs
@@ -11,7 +11,7 @@ namespace CSVTools
 
 	public class CsvToSvgLibVersion
 	{
-		private static string VersionString = "2.50";
+		private static string VersionString = "2.51";
 
 		public static string Get() { return VersionString; }
 	};
@@ -38,13 +38,18 @@ namespace CSVTools
 		// Other flags
 		public bool showMetadata = true; //noMetadata
 		public bool graphOnly = false;
-		public float compression = 0.0f;
 		public bool interactive = false;
 		public bool snapToPeaks = true; // Interactive mode only
 		public int maxHierarchyDepth = -1;
 		public char hierarchySeparator = '/';
 		public List<string> hideStatPrefixes = new List<string>();
 
+		// Compression
+		public float compression = 0.0f;
+		public bool bFixedPointGraphs = true;
+		public float fixedPointPrecisionScale = 2.0f; // Scale applied to fixed point graphs. Values greater than one allow for subpixel accuracy
+		public int lineDecimalPlaces = 3; // Legacy: only used if bFixedPointGraphs is 0
+
 		// Stacking
 		public bool stacked = false;
 		public string stackTotalStat = "";
@@ -84,7 +89,6 @@ namespace CSVTools
 
 		// Advanced params
 		public string themeName = "";
-		public int lineDecimalPlaces = 3;
 		public int frameOffset = 0;
 		public float statMultiplier = 1.0f;
 		public string minFilterStatName = null;
@@ -439,8 +443,11 @@ namespace CSVTools
 			svg.WriteLine("<stop stop-color = 'white' offset = '100%'/>");
 			svg.WriteLine("</linearGradient>");
 
-			svg.WriteLine("<clipPath id='graphArea'>");
-			svg.WriteLine("<rect  x='" + graphRect.x + "' y='" + graphRect.y + "' width='" + graphRect.width + "' height='" + graphRect.height + "'/>");
+			// Clip rect
+			svg.WriteLine("<clipPath id='graphClipRect'>");
+			// If we're using a fixed point scale for graphs, we need to apply it to the clip rect
+			float clipRectScale = (graphParams.bFixedPointGraphs && !graphParams.percentile) ? graphParams.fixedPointPrecisionScale : 1.0f;
+			svg.WriteLine("<rect  x='" + graphRect.x*clipRectScale + "' y='" + graphRect.y*clipRectScale + "' width='" + graphRect.width*clipRectScale + "' height='" + graphRect.height*clipRectScale + "'/>");
 			svg.WriteLine("</clipPath>");
 
 			svg.WriteLine("<filter id='dropShadowFilter' x='-20%' width='130%' height='130%'>");
@@ -1356,39 +1363,57 @@ namespace CSVTools
 					RawPoints = CompressPoints(RawPoints, graphParams.compression, rect, range);
 				}
 
-				string formatStr = "0";
-				if (numDecimalPlaces > 0)
+				string idString = id.Length > 0 ? "id='" + id + "'" : "";
+
+				float graphScale = 1.0f;
+				if (graphParams.bFixedPointGraphs)
 				{
-					formatStr += ".";
-					for (int i = 0; i < numDecimalPlaces; i++)
+					// Fixed point graphs are smaller because they don't include the decimal point, but we need to apply a scale to get subpixel accuracy
+					// We scale the polyline by the inverse of the scaling that we're applying to the points
+					graphScale = graphParams.fixedPointPrecisionScale;
+					float oneOverScale = 1.0f / graphScale;
+					svg.WriteLine("<polyline " + idString + " transform='scale("+ oneOverScale + ","+ oneOverScale + ")' points='");
+					foreach (Vec2 point in RawPoints)
 					{
-						formatStr += "0";
+						int x = (int)Math.Round(point.X * graphScale, 0);
+						int y = (int)Math.Round(point.Y * graphScale, 0);
+						svg.WriteFast(" " + x + "," + y);
 					}
 				}
-				string idString = id.Length > 0 ? "id='" + id + "'" : "";
-
-				svg.WriteLine("<polyline " + idString + " points='");
-				foreach (Vec2 point in RawPoints)
+				else
 				{
-					float x = point.X;
-					float y = point.Y;
-					svg.WriteFast(" " + x.ToString(formatStr) + "," + y.ToString(formatStr));
+					string formatStr = "0";
+					if (numDecimalPlaces > 0)
+					{
+						formatStr += ".";
+						for (int i = 0; i < numDecimalPlaces; i++)
+						{
+							formatStr += "0";
+						}
+					}
+					svg.WriteLine("<polyline " + idString + " points='");
+					foreach (Vec2 point in RawPoints)
+					{
+						float x = point.X;
+						float y = point.Y;
+						svg.WriteFast(" " + x.ToString(formatStr) + "," + y.ToString(formatStr));
+					}
 				}
 
 				string fillString = "none";
 				if (graphParams.stacked)
 				{
-					thickness = 0.0f;// Math.Min(thickness, 1.0f);
+					thickness = Math.Min(thickness, 0.01f);
 					float lastSample = samples[n - 1];
-					svg.WriteFast(" " + ToSvgX((float)n + 20, rect, range) + "," + ToSvgY(lastSample, rect, range) );
-					svg.WriteFast(" " + ToSvgX((float)n + 20, rect, range) + "," + ToSvgY(0.0f, rect, range) );
-					svg.WriteFast(" " + ToSvgX((float)n - 20, rect, range) + "," + ToSvgY(0.0f, rect, range) );
-					svg.WriteFast(" " + ToSvgX(start, rect, range) + "," + ToSvgY(0.0f, rect, range) );
+					svg.WriteFast(" " + ToSvgX((float)n + 20, rect, range) * graphScale + "," + ToSvgY(lastSample, rect, range) * graphScale);
+					svg.WriteFast(" " + ToSvgX((float)n + 20, rect, range) * graphScale + "," + ToSvgY(0.0f, rect, range) * graphScale);
+					svg.WriteFast(" " + ToSvgX((float)n - 20, rect, range) * graphScale + "," + ToSvgY(0.0f, rect, range) * graphScale);
+					svg.WriteFast(" " + ToSvgX(start, rect, range) * graphScale + "," + ToSvgY(0.0f, rect, range) * graphScale );
 
 					colour.alpha = 1.0f;
 					fillString = colour.SVGStringNoQuotes();
 				}
-				svg.WriteLine("' style='fill:" + fillString + ";stroke-width:" + thickness + "; clip-path: url(#graphArea)' stroke=" + colour.SVGString() + "/>");
+				svg.WriteLine("' style='fill:" + fillString + ";stroke-width:" + thickness * graphScale + "; clip-path: url(#graphClipRect)' stroke=" + colour.SVGString() + "/>");
 			}
 		}
 		void DrawPercentileGraph(SvgFile svg, List<float> samples, Colour colour, Rect rect, Range range, string id)
@@ -1425,7 +1450,7 @@ namespace CSVTools
 
 				svg.WriteFast(" " + x + "," + y);
 			}
-			svg.WriteLine("' style='fill:none;stroke-width:1.3; clip-path: url(#graphArea)' stroke=" + colour.SVGString() + "/>");
+			svg.WriteLine("' style='fill:none;stroke-width:1.3; clip-path: url(#graphClipRect)' stroke=" + colour.SVGString() + "/>");
 
 		}
 
