Shrink map summary images to the destination size rather than just copying

diff --git a/Engine/Source/Programs/CSVTools/PerfReportTool/PerfReportTool.cs b/Engine/Source/Programs/CSVTools/PerfReportTool/PerfReportTool.cs
--- a/Engine/Source/Programs/CSVTools/PerfReportTool/PerfReportTool.cs
+++ b/Engine/Source/Programs/CSVTools/PerfReportTool/PerfReportTool.cs
@@ -22,7 +22,7 @@ namespace PerfReportTool
 {
     class Version
     {
-        private static string VersionString = "4.72";
+        private static string VersionString = "4.73";
 
         public static string Get() { return VersionString; }
     };
diff --git a/Engine/Source/Programs/CSVTools/PerfReportTool/Summaries/MapOverlaySummary.cs b/Engine/Source/Programs/CSVTools/PerfReportTool/Summaries/MapOverlaySummary.cs
--- a/Engine/Source/Programs/CSVTools/PerfReportTool/Summaries/MapOverlaySummary.cs
+++ b/Engine/Source/Programs/CSVTools/PerfReportTool/Summaries/MapOverlaySummary.cs
@@ -7,6 +7,8 @@ using System.Linq;
 using System.Xml.Linq;
 using PerfReportTool;
 using CSVStats;
+using System.Drawing;
+using System.Drawing.Imaging;
 
 namespace PerfSummaries
 {
@@ -116,6 +118,19 @@ namespace PerfSummaries
 			return (int)(svgY + 0.5f);
 		}
 
+		private void CopyAndResizeImage(string sourceImagePath, string destImagePath, int destWidth, int destHeight)
+		{
+			using (FileStream fileStream = new FileStream(sourceImagePath, FileMode.Open, FileAccess.Read))
+			{
+				var image = System.Drawing.Image.FromStream(fileStream);
+				var thumbnail = image.GetThumbnailImage(destWidth, destHeight, null, IntPtr.Zero);
+				using (var destImageStream = new FileStream(destImagePath, FileMode.OpenOrCreate, FileAccess.Write))
+				{
+					thumbnail.Save(destImageStream, ImageFormat.Jpeg);
+				}
+			}
+		}
+
 		public override void WriteSummaryData(System.IO.StreamWriter htmlFile, CsvStats csvStats, bool bWriteSummaryCsv, SummaryTableRowData rowData, string htmlFileName)
 		{
 			// Output HTML
@@ -127,8 +142,10 @@ namespace PerfSummaries
 				if (!File.Exists(outputMapFilename))
 				{
 					// Copy the file to the reports folder and reset attributes to ensure it's not readonly if the source is
-					File.Copy(sourceImagePath, outputMapFilename);
-					File.SetAttributes(outputMapFilename, FileAttributes.Normal);
+					//					File.Copy(sourceImagePath, outputMapFilename);
+					//					File.SetAttributes(outputMapFilename, FileAttributes.Normal);
+					CopyAndResizeImage(sourceImagePath, outputMapFilename, (int)imageWidth, (int)imageHeight);
+
 				}
 
 				// Check if the file exists in the output directory
