CsvToSvg 2.52
- Optimize interactive graphs, removing redundant characters from number strings. Saves about 1MB in a typical 20MB report

diff --git a/Engine/Source/Programs/CSVTools/CsvToSvgLib/CsvToSvgLib.cs b/Engine/Source/Programs/CSVTools/CsvToSvgLib/CsvToSvgLib.cs
--- a/Engine/Source/Programs/CSVTools/CsvToSvgLib/CsvToSvgLib.cs
+++ b/Engine/Source/Programs/CSVTools/CsvToSvgLib/CsvToSvgLib.cs
@@ -11,7 +11,7 @@ namespace CSVTools
 
 	public class CsvToSvgLibVersion
 	{
-		private static string VersionString = "2.51";
+		private static string VersionString = "2.52";
 
 		public static string Get() { return VersionString; }
 	};
@@ -1757,6 +1757,10 @@ namespace CSVTools
 				+ "/>");
 
 		}
+		float frac(float value)
+		{
+			return value - (float)Math.Truncate(value);
+		}
 
 		string GetJSStatName(string statName)
 		{
@@ -1941,7 +1945,27 @@ namespace CSVTools
 					svg.Write("var " + statInfo.jsVarName + " = [");
 					foreach (float value in statInfo.statSamples.samples)
 					{
-						svg.WriteFast(value.ToString("0.00") + ",");
+						float fracValue = frac(value);
+						// If this is very close to a whole number, write as a whole number to save 3 bytes
+						if (fracValue <= 0.005f || fracValue >= 0.995f)
+						{
+							svg.WriteFast((int)Math.Round(value, 0) + ",");
+						}
+						else
+						{
+							string str = value.ToString("0.00");
+							if (str[0] == '0' && str[1] == '.')
+							{
+								// Trim off the leading 0 before the decimal point - it's unnecessary
+								str = str.Substring(1);
+							}
+							if (str.Last() == '0')
+							{
+								// Trim off trailing 0s
+								str = str.Substring(0,str.Length-1);
+							}
+							svg.WriteFast(str + ",");
+						}
 					}
 					svg.WriteLine("]");
 				}
