CsvToSVG 2.53 - Refactor and optimize SmoothStats via add a SmoothKernel class.
 - Add support for downsampled smoothing for an up to 16x speedup
 - More efficient threading
 - Minor perflog fixes
 - Add -PerfLog arg for detailed timings

diff --git a/Engine/Source/Programs/CSVTools/CSVToSVG/CsvToSvgProgram.cs b/Engine/Source/Programs/CSVTools/CSVToSVG/CsvToSvgProgram.cs
--- a/Engine/Source/Programs/CSVTools/CSVToSVG/CsvToSvgProgram.cs
+++ b/Engine/Source/Programs/CSVTools/CSVToSVG/CsvToSvgProgram.cs
@@ -402,6 +402,7 @@ namespace CSVTools
 			{
 				graphParams.bSmoothMultithreaded = smoothMultiThreadedStr == "1";
 			}
+			graphParams.bPerfLog = GetBoolArg("perfLog");
 
 			// Misc params (rarely used)
 			graphParams.frameOffset = GetIntArg("frameOffset", graphParams.frameOffset);
diff --git a/Engine/Source/Programs/CSVTools/CsvStats/CommandLineTool.cs b/Engine/Source/Programs/CSVTools/CsvStats/CommandLineTool.cs
--- a/Engine/Source/Programs/CSVTools/CsvStats/CommandLineTool.cs
+++ b/Engine/Source/Programs/CSVTools/CsvStats/CommandLineTool.cs
@@ -321,7 +321,7 @@ namespace CSVStats
 		public PerfLog(bool inLoggingEnabled, string inLogPrefix=null)
 		{
 			stopWatch = Stopwatch.StartNew();
-			lastTimeElapsed = 0.0;
+			previousTime = 0.0;
 			loggingEnabled = inLoggingEnabled;
 			if (inLogPrefix != null)
 			{
@@ -331,7 +331,10 @@ namespace CSVStats
 
 		public double LogTiming(string description, bool newLine = false)
 		{
-			double elapsed = stopWatch.Elapsed.TotalSeconds - lastTimeElapsed;
+
+			double currentTime = stopWatch.Elapsed.TotalSeconds;
+			double elapsed = currentTime - previousTime;
+
 			if (loggingEnabled)
 			{
 				Console.WriteLine("[PerfLog] " + logPrefix + String.Format("{0,-25} : {1,-10}", description, (elapsed * 1000.0).ToString("0.0") + "ms"), 70);
@@ -340,7 +343,7 @@ namespace CSVStats
 					Console.WriteLine();
 				}
 			}
-			lastTimeElapsed = stopWatch.Elapsed.TotalSeconds;
+			previousTime = currentTime;
 			return elapsed;
 		}
 
@@ -361,7 +364,7 @@ namespace CSVStats
 			return elapsed;
 		}
 		Stopwatch stopWatch;
-		double lastTimeElapsed;
+		double previousTime;
 		bool loggingEnabled;
 		string logPrefix = "";
 	}
diff --git a/Engine/Source/Programs/CSVTools/CsvToSvgLib/CsvToSvgLib.cs b/Engine/Source/Programs/CSVTools/CsvToSvgLib/CsvToSvgLib.cs
--- a/Engine/Source/Programs/CSVTools/CsvToSvgLib/CsvToSvgLib.cs
+++ b/Engine/Source/Programs/CSVTools/CsvToSvgLib/CsvToSvgLib.cs
@@ -5,13 +5,14 @@ using CSVStats;
 using System.Threading.Tasks;
 using System.Text;
 using System.ComponentModel;
+using System.Diagnostics;
 
 namespace CSVTools
 {
 
 	public class CsvToSvgLibVersion
 	{
-		private static string VersionString = "2.52";
+		private static string VersionString = "2.53";
 
 		public static string Get() { return VersionString; }
 	};
@@ -490,27 +491,32 @@ namespace CSVTools
 				}
 
 				SetLegend(newCsvStats, csvInfo.filename, graphParams, csvList.Count > 1);
-
-				if (graphParams.smooth)
+				if (!graphParams.stacked)
 				{
+					AssignColours(newCsvStats, theme, false, graphParams.colorOffset);
+				}
+				csvStatsList.Add(newCsvStats);
+			}
+			perfLog.LogTiming("ProcessCsvStats");
+
+
+			if (graphParams.smooth)
+			{
+				for (int i=0; i<csvStatsList.Count; i++)
+				{ 
 					int kernelSize = graphParams.smoothKernelSize;
 					if (graphParams.smoothKernelSize == -1)
 					{
 						float percent = graphParams.smoothKernelPercent;
 						percent = Math.Min(percent, 100.0f);
 						percent = Math.Max(percent, 0.0f);
-						float kernelSizeF = percent * 0.01f * (float)newCsvStats.SampleCount + 0.5f;
+						float kernelSizeF = percent * 0.01f * (float)csvStatsList[i].SampleCount + 0.5f;
 						kernelSize = (int)kernelSizeF;
 					}
-					newCsvStats = SmoothStats(newCsvStats, kernelSize, graphParams.bSmoothMultithreaded);
-				}
-				if (!graphParams.stacked)
-				{
-					AssignColours(newCsvStats, theme, false, graphParams.colorOffset);
+					csvStatsList[i] = SmoothStats(csvStatsList[i], kernelSize, graphParams.bSmoothMultithreaded);
 				}
-				csvStatsList.Add(newCsvStats);
+				perfLog.LogTiming("SmoothStats");
 			}
-			perfLog.LogTiming("ProcessCsvStats");
 
 			Range range = new Range(graphParams.minX, graphParams.maxX, graphParams.minY, graphParams.maxY);
 
@@ -766,82 +772,137 @@ namespace CSVTools
 			return newRange;
 		}
 
-
-
-		CsvStats SmoothStats(CsvStats stats, int KernelSize, bool bMultiThreaded)
+		class SmoothKernel
 		{
-			// Compute Gaussian Weights
-			float[] Weights = new float[KernelSize];
-			float TotalWeight = 0.0f;
-			for (int i = 0; i < KernelSize; i++)
+			public SmoothKernel(int inKernelSize)
 			{
-				double df = (double)i + 0.5;
-				double b = (double)KernelSize / 2;
-				double c = (double)KernelSize / 8;
-				double weight = Math.Exp(-(Math.Pow(df - b, 2.0) / Math.Pow(2 * c, 2.0)));
-				Weights[i] = (float)weight;
-				TotalWeight += (float)weight;
-			}
+				downsampleLevel = Math.Min( (int)Math.Log((double)inKernelSize / 20.0, 2.0), 4);
 
-			for (int i = 0; i < KernelSize; i++)
-			{
-				Weights[i] /= TotalWeight;
+				kernelSize = inKernelSize >> downsampleLevel;
+				weights = new float[kernelSize];
+				float totalWeight = 0.0f;
+				for (int i = 0; i < kernelSize; i++)
+				{
+					double df = (double)i + 0.5;
+					double b = (double)kernelSize / 2;
+					double c = (double)kernelSize / 8;
+					double weight = Math.Exp(-(Math.Pow(df - b, 2.0) / Math.Pow(2 * c, 2.0)));
+					weights[i] = (float)weight;
+					totalWeight += (float)weight;
+				}
+				float oneOverTotalWeight = 1.0f / totalWeight;
+				for (int i = 0; i < kernelSize; i++)
+				{
+					weights[i] *= oneOverTotalWeight;
+				}
 			}
-
-			List<StatSamples> statSampleList = stats.Stats.Values.ToList();
-
-			// Add the stats to smoothstats before the parallel for, so we can preserve the order
-			CsvStats smoothStats = new CsvStats();
-			foreach (StatSamples srcStatSamples in stats.Stats.Values)
+			public StatSamples SmoothStatSamples(StatSamples sourceStatSamples)
 			{
-				StatSamples NewStatSamples = new StatSamples(srcStatSamples, false);
-				smoothStats.AddStat(NewStatSamples);
-			}
-			StatSamples[] SmoothSamplesArray = smoothStats.Stats.Values.ToArray();
+				StatSamples destStatSamples = new StatSamples(sourceStatSamples,false);
+				destStatSamples.samples.Capacity = sourceStatSamples.samples.Count;
 
-			Action<int> smoothStat = i =>
-			{
-				StatSamples srcStatSamples = statSampleList[i];
-				StatSamples destStatSamples = SmoothSamplesArray[i];// new StatSamples(srcStatSamples);
+				// Downsample the source samples to the specified size
+				List<float> downsampledSamples = sourceStatSamples.samples;
+				for ( int i=0; i< downsampleLevel; i++)
+				{
+					int currentSampleCount = downsampledSamples.Count / 2;
+					List<float> nextLevelSamples = new List<float>(currentSampleCount);
+					for (int j=0; j< currentSampleCount; j++)
+					{
+						int prevLevelIndex = j * 2;
+						nextLevelSamples.Add( ( downsampledSamples[prevLevelIndex] + downsampledSamples[prevLevelIndex + 1] ) * 0.5f );
+					}
+					downsampledSamples = nextLevelSamples;
+				}
+				int downsampleScale = 1 << downsampleLevel;
+
+				int kernelSizeClamped = Math.Min(kernelSize, downsampledSamples.Count);
 
-				int maxSampleIndex = srcStatSamples.samples.Count - 1;
-				for (int j = 0; j < srcStatSamples.samples.Count; j++)
+				for (int i = 0; i < sourceStatSamples.samples.Count; i++)
 				{
 					float sum = 0.0f;
-					for (int k = 0; k < KernelSize; k++)
+					int startIndex = (i >> downsampleLevel) - kernelSizeClamped / 2;
+					int endIndex = (startIndex + kernelSizeClamped);
+					int startIndexClamped = Math.Max(0, startIndex);
+					int endIndexClamped = Math.Min(downsampledSamples.Count, endIndex);
+					int kernelStartIndex = startIndexClamped - startIndex;
+					int kernelIndex = kernelStartIndex;
+
+					// Smooth the in-range samples
+					for (int j = startIndexClamped; j < endIndexClamped; j++)
 					{
-						int SampleIndex = j - (KernelSize / 2) + k;
+						sum += downsampledSamples[j] * weights[kernelIndex];
+						kernelIndex++;
+					}
 
-						// If the index is out of range, mirror it (so as not to bias the result)
-						if (SampleIndex < 0)
+					// Mirror the out of range samples < 0 
+					if (startIndex < 0)
+					{
+						kernelIndex = kernelStartIndex - 1;
+						endIndexClamped = Math.Min(downsampledSamples.Count, -startIndex);
+						for (int j = startIndexClamped; j < endIndexClamped; j++)
 						{
-							SampleIndex = -SampleIndex;
+							sum += downsampledSamples[j] * weights[kernelIndex];
+							kernelIndex--;
 						}
-						else if (SampleIndex > maxSampleIndex)
+					}
+					// Mirror the out of range samples >= n
+					if (endIndex > downsampledSamples.Count)
+					{
+						kernelIndex = kernelSizeClamped - 1;
+						startIndexClamped = Math.Max(0, downsampledSamples.Count - endIndex + downsampledSamples.Count);
+						for (int j = startIndexClamped; j < endIndexClamped; j++)
 						{
-							SampleIndex = maxSampleIndex - (SampleIndex - maxSampleIndex);
+							sum += downsampledSamples[j] * weights[kernelIndex];
+							kernelIndex--;
 						}
-						// clamp the src index to a valid range
-						int srcIndex = Math.Min(Math.Max(0, SampleIndex), maxSampleIndex);
-
-						sum += srcStatSamples.samples[srcIndex] * Weights[k];
 					}
+
 					destStatSamples.samples.Add(sum);
 				}
-			};
+				return destStatSamples;
+			}
+
+			int kernelSize;
+			int downsampleLevel;
+			float [] weights;
+		};
 
+
+
+
+
+
+		CsvStats SmoothStats(CsvStats stats, int KernelSize, bool bMultiThreaded)
+		{
+			// Compute Gaussian Weights
+			SmoothKernel kernel = new SmoothKernel(KernelSize);
+
+			List<StatSamples> statSampleList = stats.Stats.Values.ToList();
+			// Add the stats to smoothstats before the parallel for, so we can preserve the order
+
+			StatSamples[] smoothSamplesArray = new StatSamples[stats.Stats.Count];
 			if (bMultiThreaded)
 			{
-				Parallel.For(0, stats.Stats.Values.Count, smoothStat);
+				int numThreads = Environment.ProcessorCount/2;
+				Parallel.For(0, stats.Stats.Values.Count, new ParallelOptions { MaxDegreeOfParallelism = numThreads }, i =>
+				{
+					smoothSamplesArray[i] = kernel.SmoothStatSamples(statSampleList[i]);
+				});
 			}
 			else
 			{
 				for (int i=0; i< stats.Stats.Values.Count; i++)
 				{
-					smoothStat(i);
+					smoothSamplesArray[i] = kernel.SmoothStatSamples(statSampleList[i]);
 				}
 			}
 
+			CsvStats smoothStats = new CsvStats();
+			foreach (StatSamples stat in smoothSamplesArray)
+			{
+				smoothStats.AddStat(stat);
+			}
 
 			smoothStats.metaData = stats.metaData;
 			smoothStats.Events = stats.Events;
