PerfReportTool 4.74
- Link templates are only added for extra relative links. Links starting with http:// or https:// are not replaced

diff --git a/Engine/Source/Programs/CSVTools/PerfReportTool/PerfReportTool.cs b/Engine/Source/Programs/CSVTools/PerfReportTool/PerfReportTool.cs
--- a/Engine/Source/Programs/CSVTools/PerfReportTool/PerfReportTool.cs
+++ b/Engine/Source/Programs/CSVTools/PerfReportTool/PerfReportTool.cs
@@ -22,7 +22,7 @@ namespace PerfReportTool
 {
     class Version
     {
-        private static string VersionString = "4.73";
+        private static string VersionString = "4.74";
 
         public static string Get() { return VersionString; }
     };
@@ -149,7 +149,7 @@ namespace PerfReportTool
 			"       -listFiles : just list all files that pass the metadata query. Don't generate any reports.\n" +
 			"       -reportLinkRootPath <path> : Make report links relative to this\n" +
 			"       -csvLinkRootPath <path> : Make CSV file links relative to this\n" +
-			"       -linkTemplates : insert templates instead of links that can be replaced later, e.g {{LinkTemplate:Report:<CSV ID>}}\n" +
+			"       -linkTemplates : insert templates in place of relative links that can be replaced later, e.g {{LinkTemplate:Report:<CSV ID>}}\n" +
 			"       -weightByColumn : weight collated table averages by this column (overrides value specified in the report XML)\n" +
 			"       -noWeightedAvg : Don't use weighted averages for the collated table\n" +
 			"       -minFrameCount <n> : ignore CSVs without at least this number of valid frames\n" +
@@ -951,7 +951,10 @@ namespace PerfReportTool
 				string csvLink = "<a href='" + csvPath + "'>" + shortName + ".csv" + "</a>";
 				if (bLinkTemplates)
 				{
-					csvLink = "{LinkTemplate:Csv:" + (csvId ?? "0") + "}";
+					if ( !csvPath.StartsWith("http://") && !csvPath.StartsWith("https://") )
+					{
+						csvLink = "{LinkTemplate:Csv:" + (csvId ?? "0") + "}";
+					}
 				}
 
 				rowData.Add(SummaryTableElement.Type.ToolMetadata, "Csv File", csvLink, null, csvPath);
@@ -968,7 +971,10 @@ namespace PerfReportTool
 					string reportLink = "<a href='" + htmlUrl + "'>Link</a>";
 					if (bLinkTemplates)
 					{
-						reportLink = "{LinkTemplate:Report:" + (csvId ?? "0") + "}";
+						if (!htmlUrl.StartsWith("http://") && !htmlUrl.StartsWith("https://"))
+						{
+							reportLink = "{LinkTemplate:Report:" + (csvId ?? "0") + "}";
+						}
 					}	
 					rowData.Add(SummaryTableElement.Type.ToolMetadata, "Report", reportLink);
 				}
diff --git a/Engine/Source/Programs/CSVTools/PerfReportTool/Summaries/ExtraLinksSummary.cs b/Engine/Source/Programs/CSVTools/PerfReportTool/Summaries/ExtraLinksSummary.cs
--- a/Engine/Source/Programs/CSVTools/PerfReportTool/Summaries/ExtraLinksSummary.cs
+++ b/Engine/Source/Programs/CSVTools/PerfReportTool/Summaries/ExtraLinksSummary.cs
@@ -24,8 +24,11 @@ namespace PerfSummaries
                 ReportText = Sections[0];
                 SummaryTableText = Sections[1];
 				LinkURL = Sections[2];
-                LinkTemplateCsvId = inLinkTemplateCsvId;
-            }
+				if (!LinkURL.StartsWith("http://") && !LinkURL.StartsWith("https://"))
+				{
+					LinkTemplateCsvId = inLinkTemplateCsvId;
+				}
+			}
             public string GetLinkString(bool bIsSummaryTable)
             {
                 if (bIsSummaryTable && LinkTemplateCsvId != null)
