PerfReportTool 4.75 - Add logging around map image handling

diff --git a/Engine/Source/Programs/CSVTools/PerfReportTool/PerfReportTool.cs b/Engine/Source/Programs/CSVTools/PerfReportTool/PerfReportTool.cs
--- a/Engine/Source/Programs/CSVTools/PerfReportTool/PerfReportTool.cs
+++ b/Engine/Source/Programs/CSVTools/PerfReportTool/PerfReportTool.cs
@@ -22,7 +22,7 @@ namespace PerfReportTool
 {
     class Version
     {
-        private static string VersionString = "4.74";
+        private static string VersionString = "4.75";
 
         public static string Get() { return VersionString; }
     };
diff --git a/Engine/Source/Programs/CSVTools/PerfReportTool/Summaries/MapOverlaySummary.cs b/Engine/Source/Programs/CSVTools/PerfReportTool/Summaries/MapOverlaySummary.cs
--- a/Engine/Source/Programs/CSVTools/PerfReportTool/Summaries/MapOverlaySummary.cs
+++ b/Engine/Source/Programs/CSVTools/PerfReportTool/Summaries/MapOverlaySummary.cs
@@ -120,12 +120,16 @@ namespace PerfSummaries
 
 		private void CopyAndResizeImage(string sourceImagePath, string destImagePath, int destWidth, int destHeight)
 		{
+			Console.WriteLine("Downsampling map image.\n  Source: " + sourceImagePath+"\n  Dest  : "+destImagePath);
 			using (FileStream fileStream = new FileStream(sourceImagePath, FileMode.Open, FileAccess.Read))
 			{
+				Console.WriteLine("Reading source image");
 				var image = System.Drawing.Image.FromStream(fileStream);
+				Console.WriteLine("Generating downsampled image");
 				var thumbnail = image.GetThumbnailImage(destWidth, destHeight, null, IntPtr.Zero);
 				using (var destImageStream = new FileStream(destImagePath, FileMode.OpenOrCreate, FileAccess.Write))
 				{
+					Console.WriteLine("Saving downsampled map image: " + destImageStream);
 					thumbnail.Save(destImageStream, ImageFormat.Jpeg);
 				}
 			}
@@ -139,13 +143,23 @@ namespace PerfSummaries
 				string outputDirectory = System.IO.Path.GetDirectoryName(System.IO.Path.GetFullPath(htmlFileName));
 				string outputMapFilename = System.IO.Path.Combine(outputDirectory, destImageFilename);
 
-				if (!File.Exists(outputMapFilename))
+				if (File.Exists(outputMapFilename))
 				{
-					// Copy the file to the reports folder and reset attributes to ensure it's not readonly if the source is
-					//					File.Copy(sourceImagePath, outputMapFilename);
-					//					File.SetAttributes(outputMapFilename, FileAttributes.Normal);
-					CopyAndResizeImage(sourceImagePath, outputMapFilename, (int)imageWidth, (int)imageHeight);
-
+					Console.WriteLine("Map file already exists. Skipping copy: "+ outputMapFilename);
+				}
+				else
+				{
+					if (File.Exists(sourceImagePath))
+					{
+						// Copy the file to the reports folder and reset attributes to ensure it's not readonly if the source is
+						//					File.Copy(sourceImagePath, outputMapFilename);
+						//					File.SetAttributes(outputMapFilename, FileAttributes.Normal);
+						CopyAndResizeImage(sourceImagePath, outputMapFilename, (int)imageWidth, (int)imageHeight);
+					}
+					else
+					{
+						Console.WriteLine("[Warning] Can't find source map image: " + sourceImagePath);
+					}
 				}
 
 				// Check if the file exists in the output directory
