CsvInfo 1.04 - add -showTotals option.
- display sample count and output it to json
- add proper support for totals and averages in json. When these are requested, we output stats as a dict rather than an array

diff --git a/Engine/Source/Programs/CSVTools/CsvInfo/CSVInfo.cs b/Engine/Source/Programs/CSVTools/CsvInfo/CSVInfo.cs
--- a/Engine/Source/Programs/CSVTools/CsvInfo/CSVInfo.cs
+++ b/Engine/Source/Programs/CSVTools/CsvInfo/CSVInfo.cs
@@ -12,7 +12,7 @@ namespace CSVInfo
 {
     class Version
     {
-        private static string VersionString = "1.03";
+        private static string VersionString = "1.04";
 
         public static string Get() { return VersionString; }
     };
@@ -89,6 +89,7 @@ namespace CSVInfo
             ReadCommandLine(args);
 
             bool showAverages = GetBoolArg("showAverages");
+			bool showTotals= GetBoolArg("showTotals");
 			string jsonFilename = GetArg("toJson",false);
 
 			CsvStats csvStats = CsvStats.ReadCSVFile(csvFilename, null);
@@ -98,25 +99,14 @@ namespace CSVInfo
 				return;
 			}
 
-			List<string> statLines = new List<string>();
-			foreach (StatSamples stat in csvStats.Stats.Values.ToArray())
-			{
-				string statLine = stat.Name;
-				if (showAverages)
-				{
-					statLine += " (" + stat.average.ToString() + ")";
-				}
-				statLines.Add(statLine);
-			}
-			statLines.Sort();
-
-
 			if (jsonFilename != "")
 			{
 				// We just write the lines raw, since this version of .Net doesn't have a json serializer. 
 				// TODO: Fix this when we upgrade to .Net 5.0 and use System.Text.Json
 				List<string> jsonLines = new List<string>();
 				jsonLines.Add("{");
+				jsonLines.Add("  \"sampleCount\":" + csvStats.SampleCount+",");
+
 				if (csvStats.metaData != null)
 				{
 					jsonLines.Add("  \"metadata\": {");
@@ -135,7 +125,41 @@ namespace CSVInfo
 					}
 					jsonLines.Add("  },");
 				}
-				jsonLines.Add("  \"stats\": "+ ToJsonStringList(statLines));
+				List<string> statLines = new List<string>();
+				foreach (StatSamples stat in csvStats.Stats.Values.ToArray())
+				{
+					statLines.Add(stat.Name);
+				}
+				statLines.Sort();
+
+				if (showTotals || showAverages)
+				{
+					jsonLines.Add("  \"stats\": {");
+					for (int i=0; i<statLines.Count; i++)
+					{
+						string statName = statLines[i];
+						List<string> entries = new List<string>();
+						if ( showTotals )
+						{
+							entries.Add("\"total\":" + csvStats.GetStat(statName).total);
+						}
+						if (showAverages)
+						{
+							entries.Add("\"average\":" + csvStats.GetStat(statName).average);
+						}
+						string line = "    \"" + statName + "\": {" + String.Join(",", entries) + "}";
+						if (i < statLines.Count-1)
+							line += ",";
+						jsonLines.Add(line);
+					}
+					jsonLines.Add("  }");
+				}
+				else
+				{
+					// Just output stats as an array if totals/averages were not requested
+					jsonLines.Add("  \"stats\": " + ToJsonStringList(statLines));
+				}
+
 
 				jsonLines.Add("}");
 				System.IO.File.WriteAllLines(jsonFilename,jsonLines);
@@ -143,6 +167,25 @@ namespace CSVInfo
 			}
 			else
 			{
+				// Write out the sample count
+				Console.Out.WriteLine("Sample Count: " + csvStats.SampleCount);
+
+				List<string> statLines = new List<string>();
+				foreach (StatSamples stat in csvStats.Stats.Values.ToArray())
+				{
+					string statLine = stat.Name;
+					if (showAverages)
+					{
+						statLine += " (" + stat.average.ToString() + ") ";
+					}
+					if (showTotals)
+					{
+						statLine += " (Total: " + stat.total.ToString()+") ";
+					}
+					statLines.Add(statLine);
+				}
+				statLines.Sort();
+
 				// Write out the sorted stat names
 				Console.Out.WriteLine("Stats:");
 				foreach (string statLine in statLines)
