CsvInfo 1.05 - don't read the full CSV unless it's necessary. This is roughly 100x times faster for binary CSVs, and 2x as fast for text
- add -forceFullRead option to force a full read (useful for validation)

diff --git a/Engine/Source/Programs/CSVTools/CsvInfo/CSVInfo.cs b/Engine/Source/Programs/CSVTools/CsvInfo/CSVInfo.cs
--- a/Engine/Source/Programs/CSVTools/CsvInfo/CSVInfo.cs
+++ b/Engine/Source/Programs/CSVTools/CsvInfo/CSVInfo.cs
@@ -12,7 +12,7 @@ namespace CSVInfo
 {
     class Version
     {
-        private static string VersionString = "1.04";
+        private static string VersionString = "1.05";
 
         public static string Get() { return VersionString; }
     };
@@ -73,6 +73,7 @@ namespace CSVInfo
                 "Format: \n" +
                 "  <csvfilename>\n"+
 				"  [-showaverages]\n"+
+				"  [-forceFullRead] (always reads the full CSV)\n" +
 				"  [-quiet] (no logging. Just throws returns a non-zero error code if the CSV is bad)\n" +
 				"  [-toJson <filename>]";
 
@@ -91,8 +92,10 @@ namespace CSVInfo
             bool showAverages = GetBoolArg("showAverages");
 			bool showTotals= GetBoolArg("showTotals");
 			string jsonFilename = GetArg("toJson",false);
+			bool bReadJustHeader = !GetBoolArg("forceFullRead") && !showAverages && !showTotals;
 
-			CsvStats csvStats = CsvStats.ReadCSVFile(csvFilename, null);
+			CSVStats.CsvFileInfo fileInfo = new CsvFileInfo();
+			CsvStats csvStats = CsvStats.ReadCSVFile(csvFilename, null, 0, false, fileInfo, bReadJustHeader);
 
 			if ( GetBoolArg("quiet") )
 			{
@@ -105,7 +108,7 @@ namespace CSVInfo
 				// TODO: Fix this when we upgrade to .Net 5.0 and use System.Text.Json
 				List<string> jsonLines = new List<string>();
 				jsonLines.Add("{");
-				jsonLines.Add("  \"sampleCount\":" + csvStats.SampleCount+",");
+				jsonLines.Add("  \"sampleCount\":" + fileInfo.SampleCount+",");
 
 				if (csvStats.metaData != null)
 				{
@@ -168,7 +171,7 @@ namespace CSVInfo
 			else
 			{
 				// Write out the sample count
-				Console.Out.WriteLine("Sample Count: " + csvStats.SampleCount);
+				Console.Out.WriteLine("Sample Count: " + fileInfo.SampleCount);
 
 				List<string> statLines = new List<string>();
 				foreach (StatSamples stat in csvStats.Stats.Values.ToArray())
diff --git a/Engine/Source/Programs/CSVTools/CsvStats/CsvStats.cs b/Engine/Source/Programs/CSVTools/CsvStats/CsvStats.cs
--- a/Engine/Source/Programs/CSVTools/CsvStats/CsvStats.cs
+++ b/Engine/Source/Programs/CSVTools/CsvStats/CsvStats.cs
@@ -48,6 +48,7 @@ namespace CSVStats
 		public bool bIsCsvBin;
 		public CsvBinVersion BinVersion = CsvBinVersion.COUNT;
 		public CsvBinCompressionLevel BinCompressionLevel = CsvBinCompressionLevel.None;
+		public int SampleCount = 0;
 	};
 
 
@@ -886,6 +887,7 @@ namespace CSVStats
 				FileInfoOut.bIsCsvBin = true;
 				FileInfoOut.BinVersion = (CsvBinVersion)version;
 				FileInfoOut.BinCompressionLevel = compressionLevel;
+				FileInfoOut.SampleCount = sampleCount;
 			}
 
 			if (justHeader)
@@ -1548,21 +1550,17 @@ namespace CSVStats
 		}
 
 
-		public static CsvStats ReadCSVFile(string csvFilename, string[] statNames, int numRowsToSkip = 0, bool bGenerateCsvIdIfMissing=false, CsvFileInfo FileInfoOut=null)
+		public static CsvStats ReadCSVFile(string csvFilename, string[] statNames, int numRowsToSkip = 0, bool bGenerateCsvIdIfMissing=false, CsvFileInfo FileInfoOut=null, bool bJustHeader=false)
         {
 			CsvStats statsOut;
 			if (csvFilename.EndsWith(".csv.bin"))
 			{
-				statsOut = ReadBinFile(csvFilename, statNames, numRowsToSkip, false, FileInfoOut);
+				statsOut = ReadBinFile(csvFilename, statNames, numRowsToSkip, bJustHeader, FileInfoOut);
 			}
 			else
 			{
 				string[] lines = ReadLinesFromFile(csvFilename);
-				statsOut = ReadCSVFromLines(lines, statNames, numRowsToSkip);
-				if ( FileInfoOut != null )
-				{
-					FileInfoOut.bIsCsvBin = false;
-				}
+				statsOut = ReadCSVFromLines(lines, statNames, numRowsToSkip, bJustHeader, FileInfoOut);
 			}
 			if (bGenerateCsvIdIfMissing)
 			{
@@ -1646,7 +1644,7 @@ namespace CSVStats
             return false;
         }
 
-        public static CsvStats ReadCSVFromLines(string[] linesArray, string[] statNames, int numRowsToSkip = 0, bool skipReadingData=false)
+        public static CsvStats ReadCSVFromLines(string[] linesArray, string[] statNames, int numRowsToSkip = 0, bool skipReadingData=false, CsvFileInfo fileInfo=null)
         {
             List<string> lines = linesArray.ToList();
 
@@ -1684,10 +1682,18 @@ namespace CSVStats
                 lines.RemoveAt(lines.Count - 1);
             }
 
-            if (skipReadingData)
+			// First line is headings, last line contains build info 
+			int numSamples = lines.Count - (bHasMetaData ? 2 : 1);
+			if (fileInfo != null)
+			{
+				fileInfo.bIsCsvBin = false;
+				fileInfo.SampleCount = numSamples;
+			}
+
+			if (skipReadingData)
             {
-				int dataLineCount = bHasMetaData ? lines.Count-2 : lines.Count-1;
-                lines.RemoveRange(1, dataLineCount);
+				lines.RemoveRange(1, numSamples);
+				numSamples = 0;
             }
 
 			// Get the list of lower case stat names, expanding wildcards
@@ -1718,9 +1724,6 @@ namespace CSVStats
                 statNamesLowercase = newStatNamesLowercase.ToArray();
             }
 
-            // First line is headings, last line contains build info 
-            int numSamples = lines.Count - (bHasMetaData ? 2 : 1);
-
             // Create the stats
             int eventHeadingIndex = -1;
             StatSamples[] stats = new StatSamples[headings.Length];
