CsvToSvg 2.54 - fix bug with graph smoothing when kernel size is small: Clamp the downsample level so it can't go negative

diff --git a/Engine/Source/Programs/CSVTools/CsvToSvgLib/CsvToSvgLib.cs b/Engine/Source/Programs/CSVTools/CsvToSvgLib/CsvToSvgLib.cs
--- a/Engine/Source/Programs/CSVTools/CsvToSvgLib/CsvToSvgLib.cs
+++ b/Engine/Source/Programs/CSVTools/CsvToSvgLib/CsvToSvgLib.cs
@@ -12,7 +12,7 @@ namespace CSVTools
 
 	public class CsvToSvgLibVersion
 	{
-		private static string VersionString = "2.53";
+		private static string VersionString = "2.54";
 
 		public static string Get() { return VersionString; }
 	};
@@ -776,7 +776,7 @@ namespace CSVTools
 		{
 			public SmoothKernel(int inKernelSize)
 			{
-				downsampleLevel = Math.Min( (int)Math.Log((double)inKernelSize / 20.0, 2.0), 4);
+				downsampleLevel = Math.Max( Math.Min( (int)Math.Log((double)inKernelSize / 20.0, 2.0), 4), 0);
 
 				kernelSize = inKernelSize >> downsampleLevel;
 				weights = new float[kernelSize];
