CsvToSVG 2.55 - fix bug where colors are all the same with multiple CSVs

diff --git a/Engine/Source/Programs/CSVTools/CsvToSvgLib/CsvToSvgLib.cs b/Engine/Source/Programs/CSVTools/CsvToSvgLib/CsvToSvgLib.cs
--- a/Engine/Source/Programs/CSVTools/CsvToSvgLib/CsvToSvgLib.cs
+++ b/Engine/Source/Programs/CSVTools/CsvToSvgLib/CsvToSvgLib.cs
@@ -12,7 +12,7 @@ namespace CSVTools
 
 	public class CsvToSvgLibVersion
 	{
-		private static string VersionString = "2.54";
+		private static string VersionString = "2.55";
 
 		public static string Get() { return VersionString; }
 	};
@@ -465,6 +465,7 @@ namespace CSVTools
 
 			// Generate the list of processed, filtered CSV stats for graphing
 			List<CsvStats> csvStatsList = new List<CsvStats>();
+			int currentColorOffset = graphParams.colorOffset;
 			foreach (CsvInfo csvInfo in csvList)
 			{
 				CsvStats newCsvStats = ProcessCsvStats(csvInfo.stats, graphParams);
@@ -493,7 +494,7 @@ namespace CSVTools
 				SetLegend(newCsvStats, csvInfo.filename, graphParams, csvList.Count > 1);
 				if (!graphParams.stacked)
 				{
-					AssignColours(newCsvStats, theme, false, graphParams.colorOffset);
+					currentColorOffset = AssignColours(newCsvStats, theme, false, currentColorOffset);
 				}
 				csvStatsList.Add(newCsvStats);
 			}
@@ -937,7 +938,7 @@ namespace CSVTools
 			}
 		}
 
-		void AssignColours(CsvStats stats, Theme theme, bool reverseOrder, int inColorOffset)
+		int AssignColours(CsvStats stats, Theme theme, bool reverseOrder, int inColorOffset)
 		{
 			int colorOffset = inColorOffset;
 			if (reverseOrder)
@@ -960,6 +961,7 @@ namespace CSVTools
 					colorOffset++;
 				}
 			}
+			return colorOffset;
 		}
 		CsvStats StackStats(CsvStats stats, Range range, bool stackedUnsorted, string stackTotalStat, bool bStackTotalStatIsAutomatic)
 		{
