PerfReportTool 4.76 : Add support for unstripped stats in FPSCharts.
- FPSChart stats can now be tagged as unstripped. Add "unstripped" in the stat bracket parameters, space-separated to enable, e.g MemFreeMB (min unstripped).
- Simplify column handling code
- Don't generate CsvToSvg commandline args if we're not running as a separate process

Note: The unstripped parameter does not affect the column name. E.g MemFreeMB(Min Unstripped) will be output as MemFreeMB Min.

diff --git a/Engine/Source/Programs/CSVTools/PerfReportTool/PerfReportTool.cs b/Engine/Source/Programs/CSVTools/PerfReportTool/PerfReportTool.cs
--- a/Engine/Source/Programs/CSVTools/PerfReportTool/PerfReportTool.cs
+++ b/Engine/Source/Programs/CSVTools/PerfReportTool/PerfReportTool.cs
@@ -22,7 +22,7 @@ namespace PerfReportTool
 {
     class Version
     {
-        private static string VersionString = "4.75";
+        private static string VersionString = "4.76";
 
         public static string Get() { return VersionString; }
     };
@@ -718,16 +718,16 @@ namespace PerfReportTool
 			}
 
 			int numFramesStripped;
-			CsvStats unstrippedCsvStats;
-			CsvStats csvStats = ProcessCsv(csvFile, out numFramesStripped, out unstrippedCsvStats, minX, maxX, null, bStripStatsByEvents);
+			CsvStats csvStatsUnstripped;
+			CsvStats csvStats = ProcessCsv(csvFile, out numFramesStripped, out csvStatsUnstripped, minX, maxX, null, bStripStatsByEvents);
 			csvStats.WriteToCSV(outCsvFilename, false);
 		}
 
-		CsvStats ProcessCsv(CachedCsvFile csvFile, out int numFramesStripped, out CsvStats unstrippedCsvStats, int minX=0, int maxX=Int32.MaxValue, PerfLog perfLog=null, bool bStripStatsByEvents = true)
+		CsvStats ProcessCsv(CachedCsvFile csvFile, out int numFramesStripped, out CsvStats csvStatsUnstripped, int minX=0, int maxX=Int32.MaxValue, PerfLog perfLog=null, bool bStripStatsByEvents = true)
 		{
 			numFramesStripped = 0;
 			CsvStats csvStats = ReadCsvStats(csvFile, minX, maxX);
-			unstrippedCsvStats = csvStats;
+			csvStatsUnstripped = csvStats;
 			if (perfLog != null)
 			{
 				perfLog.LogTiming("    ReadCsvStats");
@@ -735,7 +735,7 @@ namespace PerfReportTool
 
 			if (bStripStatsByEvents)
 			{
-				CsvStats strippedCsvStats = StripCsvStatsByEvents(unstrippedCsvStats, out numFramesStripped);
+				CsvStats strippedCsvStats = StripCsvStatsByEvents(csvStatsUnstripped, out numFramesStripped);
 				csvStats = strippedCsvStats;
 			}
 			if (perfLog != null)
@@ -832,8 +832,6 @@ namespace PerfReportTool
 						if (graphGenerator != null)
 						{
 							GraphParams graphParams=GetCsvToSvgGraphParams(csvFile.filename, graph, thickness, minX, maxX, false, svgFilenames.Count);
-							string args = GetCsvToSvgArgs(csvFile.filename, svgFilename, graph, thickness, minX, maxX, false, svgFilenames.Count);
-
 							if (bCsvToSvgMultiThreaded)
 							{
 								csvToSvgTasks.Add( graphGenerator.MakeGraphAsync(graphParams, svgFilename, true, false) );
@@ -891,8 +889,8 @@ namespace PerfReportTool
 
 			// Read the full csv while we wait for the graph processes to complete (this is only safe for CsvToSVG processes, not task threads)
 			int numFramesStripped;
-			CsvStats unstrippedCsvStats;
-			CsvStats csvStats=ProcessCsv(csvFile, out numFramesStripped, out unstrippedCsvStats, minX, maxX, perfLog, bStripStatsByEvents);
+			CsvStats csvStatsUnstripped;
+			CsvStats csvStats=ProcessCsv(csvFile, out numFramesStripped, out csvStatsUnstripped, minX, maxX, perfLog, bStripStatsByEvents);
 
             if ( writeDetailedReport && csvToSvgProcesses.Count > 0)
             { 
@@ -1004,7 +1002,7 @@ namespace PerfReportTool
             }
 
             // Write the report
-            WriteReport(htmlFilename, title, svgFilenames, reportTypeInfo, csvStats, unstrippedCsvStats, numFramesStripped, minX, maxX, bBulkMode, rowData);
+            WriteReport(htmlFilename, title, svgFilenames, reportTypeInfo, csvStats, csvStatsUnstripped, numFramesStripped, minX, maxX, bBulkMode, rowData);
             perfLog.LogTiming("    WriteReport");
 
             // Delete the temp files
@@ -1110,7 +1108,7 @@ namespace PerfReportTool
       
 
 
-        void WriteReport(string htmlFilename, string title, List<string> svgFilenames, ReportTypeInfo reportTypeInfo, CsvStats csvStats, CsvStats unstrippedCsvStats, int numFramesStripped, int minX, int maxX, bool bBulkMode, SummaryTableRowData summaryRowData)
+        void WriteReport(string htmlFilename, string title, List<string> svgFilenames, ReportTypeInfo reportTypeInfo, CsvStats csvStats, CsvStats csvStatsUnstripped, int numFramesStripped, int minX, int maxX, bool bBulkMode, SummaryTableRowData summaryRowData)
         {
  
             ReportGraph[] graphs = reportTypeInfo.graphs.ToArray();
@@ -1206,7 +1204,7 @@ namespace PerfReportTool
 			PeakSummary peakSummary = null;
             foreach (Summary summary in summaries)
             {
-                summary.WriteSummaryData(htmlFile, summary.useUnstrippedCsvStats ? unstrippedCsvStats : csvStats, bWriteSummaryCsv, summaryRowData, htmlFilename);
+                summary.WriteSummaryData(htmlFile, summary.useUnstrippedCsvStats ? csvStatsUnstripped : csvStats, csvStatsUnstripped, bWriteSummaryCsv, summaryRowData, htmlFilename);
                 if ( summary.GetType() == typeof(PeakSummary) )
                 {
                     peakSummary = (PeakSummary)summary;
@@ -1298,13 +1296,13 @@ namespace PerfReportTool
 				string ForEmail = GetArg("foremail", false);
 				if (ForEmail != "")
 				{
-					WriteEmail(htmlFilename, title, svgFilenames, reportTypeInfo, csvStats, minX, maxX, bBulkMode);
+					WriteEmail(htmlFilename, title, svgFilenames, reportTypeInfo, csvStats, csvStatsUnstripped, minX, maxX, bBulkMode);
 				}
 			}
         }
 
 
-        void WriteEmail(string htmlFilename, string title, List<string> svgFilenames, ReportTypeInfo reportTypeInfo, CsvStats csvStats, int minX, int maxX, bool bBulkMode)
+        void WriteEmail(string htmlFilename, string title, List<string> svgFilenames, ReportTypeInfo reportTypeInfo, CsvStats csvStats, CsvStats csvStatsUnstripped, int minX, int maxX, bool bBulkMode)
         {
 			if (htmlFilename==null)
 			{
@@ -1355,7 +1353,7 @@ namespace PerfReportTool
 			PeakSummary peakSummary = null;
             foreach (Summary summary in reportTypeInfo.summaries)
             {
-                summary.WriteSummaryData(htmlFile, csvStats, bWriteSummaryCsv, null, htmlFilename);
+                summary.WriteSummaryData(htmlFile, csvStats, csvStatsUnstripped, bWriteSummaryCsv, null, htmlFilename);
                 if (summary.GetType() == typeof(PeakSummary))
                 {
                     peakSummary = (PeakSummary)summary;
diff --git a/Engine/Source/Programs/CSVTools/PerfReportTool/Summaries/BoundedStatValuesSummary.cs b/Engine/Source/Programs/CSVTools/PerfReportTool/Summaries/BoundedStatValuesSummary.cs
--- a/Engine/Source/Programs/CSVTools/PerfReportTool/Summaries/BoundedStatValuesSummary.cs
+++ b/Engine/Source/Programs/CSVTools/PerfReportTool/Summaries/BoundedStatValuesSummary.cs
@@ -90,7 +90,7 @@ namespace PerfSummaries
 
 		public override string GetName() { return "boundedstatvalues"; }
 
-		public override void WriteSummaryData(System.IO.StreamWriter htmlFile, CsvStats csvStats, bool bWriteSummaryCsv, SummaryTableRowData rowData, string htmlFileName)
+		public override void WriteSummaryData(System.IO.StreamWriter htmlFile, CsvStats csvStats, CsvStats csvStatsUnstripped, bool bWriteSummaryCsv, SummaryTableRowData rowData, string htmlFileName)
 		{
 			int startFrame = -1;
 			int endFrame = int.MaxValue;
diff --git a/Engine/Source/Programs/CSVTools/PerfReportTool/Summaries/BucketSummary.cs b/Engine/Source/Programs/CSVTools/PerfReportTool/Summaries/BucketSummary.cs
--- a/Engine/Source/Programs/CSVTools/PerfReportTool/Summaries/BucketSummary.cs
+++ b/Engine/Source/Programs/CSVTools/PerfReportTool/Summaries/BucketSummary.cs
@@ -42,7 +42,7 @@ namespace PerfSummaries
 
 		public override string GetName() { return "bucketsummary"; }
 
-		public override void WriteSummaryData(System.IO.StreamWriter htmlFile, CsvStats csvStats, bool bWriteSummaryCsv, SummaryTableRowData metadata, string htmlFileName)
+		public override void WriteSummaryData(System.IO.StreamWriter htmlFile, CsvStats csvStats, CsvStats csvStatsUnstripped, bool bWriteSummaryCsv, SummaryTableRowData metadata, string htmlFileName)
 		{
 			// Only HTML reporting is supported (does not summary table row data)
 			if (htmlFile == null)
diff --git a/Engine/Source/Programs/CSVTools/PerfReportTool/Summaries/EventSummary.cs b/Engine/Source/Programs/CSVTools/PerfReportTool/Summaries/EventSummary.cs
--- a/Engine/Source/Programs/CSVTools/PerfReportTool/Summaries/EventSummary.cs
+++ b/Engine/Source/Programs/CSVTools/PerfReportTool/Summaries/EventSummary.cs
@@ -23,7 +23,7 @@ namespace PerfSummaries
 
 		public override string GetName() { return "event"; }
 
-		public override void WriteSummaryData(System.IO.StreamWriter htmlFile, CsvStats csvStats, bool bWriteSummaryCsv, SummaryTableRowData rowData, string htmlFileName)
+		public override void WriteSummaryData(System.IO.StreamWriter htmlFile, CsvStats csvStats, CsvStats csvStatsUnstripped, bool bWriteSummaryCsv, SummaryTableRowData rowData, string htmlFileName)
 		{
 			Dictionary<string, int> eventCountsDict = new Dictionary<string, int>();
 			int eventCount = 0;
diff --git a/Engine/Source/Programs/CSVTools/PerfReportTool/Summaries/ExtraLinksSummary.cs b/Engine/Source/Programs/CSVTools/PerfReportTool/Summaries/ExtraLinksSummary.cs
--- a/Engine/Source/Programs/CSVTools/PerfReportTool/Summaries/ExtraLinksSummary.cs
+++ b/Engine/Source/Programs/CSVTools/PerfReportTool/Summaries/ExtraLinksSummary.cs
@@ -58,7 +58,7 @@ namespace PerfSummaries
 
 		public override string GetName() { return "extralinks"; }
 
-		public override void WriteSummaryData(System.IO.StreamWriter htmlFile, CsvStats csvStats, bool bWriteSummaryCsv, SummaryTableRowData rowData, string htmlFileName)
+		public override void WriteSummaryData(System.IO.StreamWriter htmlFile, CsvStats csvStats, CsvStats csvStatsUnstripped, bool bWriteSummaryCsv, SummaryTableRowData rowData, string htmlFileName)
 		{
 			List<ExtraLink> links = new List<ExtraLink>();
 
diff --git a/Engine/Source/Programs/CSVTools/PerfReportTool/Summaries/FPSChartSummary.cs b/Engine/Source/Programs/CSVTools/PerfReportTool/Summaries/FPSChartSummary.cs
--- a/Engine/Source/Programs/CSVTools/PerfReportTool/Summaries/FPSChartSummary.cs
+++ b/Engine/Source/Programs/CSVTools/PerfReportTool/Summaries/FPSChartSummary.cs
@@ -123,8 +123,33 @@ namespace PerfSummaries
 			return outData;
 		}
 
+		class ColumnInfo
+		{
+			public ColumnInfo(string inName, double inValue, ColourThresholdList inColorThresholds, string inDetails=null, bool inIsAvgValue=false)
+			{
+				Name = inName;
+				Value = inValue;
+				ColorThresholds = inColorThresholds; 
+				isAverageValue = inIsAvgValue;
+				Details = inDetails;
+
+				UpdateColor();
+			}
+
+			public void UpdateColor()
+			{
+				Color = ColourThresholdList.GetSafeColourForValue(ColorThresholds, Value);
+			}
 
-		public override void WriteSummaryData(System.IO.StreamWriter htmlFile, CsvStats csvStats, bool bWriteSummaryCsv, SummaryTableRowData rowData, string htmlFileName)
+			public string Name;
+			public double Value;
+			public string Color;
+			public ColourThresholdList ColorThresholds;
+			public string Details;
+			public bool isAverageValue;
+		};
+
+		public override void WriteSummaryData(System.IO.StreamWriter htmlFile, CsvStats csvStats, CsvStats csvStatsUnstripped, bool bWriteSummaryCsv, SummaryTableRowData rowData, string htmlFileName)
 		{
 			System.IO.StreamWriter statsCsvFile = null;
 			if (bWriteSummaryCsv)
@@ -137,81 +162,81 @@ namespace PerfSummaries
 			FpsChartData fpsChartData = ComputeFPSChartDataForFrames(frameTimes, true);
 
 			// Write the averages
-			List<string> ColumnNames = new List<string>();
-			List<double> ColumnValues = new List<double>();
-			List<string> ColumnColors = new List<string>();
-			List<ColourThresholdList> ColumnColorThresholds = new List<ColourThresholdList>();
-
-			ColumnNames.Add("Total Time (s)");
-			ColumnColorThresholds.Add(new ColourThresholdList());
-			ColumnValues.Add(fpsChartData.TotalTimeSeconds);
-			ColumnColors.Add(ColourThresholdList.GetSafeColourForValue(ColumnColorThresholds.Last(), ColumnValues.Last()));
-
-			ColumnNames.Add("Hitches/Min");
-			ColumnColorThresholds.Add(GetStatColourThresholdList(ColumnNames.Last()));
-			ColumnValues.Add(fpsChartData.HitchesPerMinute);
-			ColumnColors.Add(ColourThresholdList.GetSafeColourForValue(ColumnColorThresholds.Last(), ColumnValues.Last()));
+			List<ColumnInfo> Columns = new List<ColumnInfo>();
 
+			Columns.Add( new ColumnInfo("Total Time (s)", fpsChartData.TotalTimeSeconds, new ColourThresholdList() ) );
+			Columns.Add(new ColumnInfo("Hitches/Min", fpsChartData.HitchesPerMinute, GetStatColourThresholdList("Hitches/Min") ) );
+			
 			if (!bIgnoreHitchTimePercent)
 			{
-				ColumnNames.Add("HitchTimePercent");
-				ColumnColorThresholds.Add(GetStatColourThresholdList(ColumnNames.Last()));
-				ColumnValues.Add(fpsChartData.HitchTimePercent);
-				ColumnColors.Add(ColourThresholdList.GetSafeColourForValue(ColumnColorThresholds.Last(), ColumnValues.Last()));
+				Columns.Add(new ColumnInfo("HitchTimePercent", fpsChartData.HitchTimePercent, GetStatColourThresholdList("HitchTimePercent")));
 			}
 
+			string MvpStatName = "MVP" + fps.ToString();
 			if (!bIgnoreMVP)
 			{
-				ColumnNames.Add("MVP" + fps.ToString());
-				ColumnColorThresholds.Add(GetStatColourThresholdList(ColumnNames.Last()));
-				ColumnValues.Add(fpsChartData.MVP);
-				ColumnColors.Add(ColourThresholdList.GetSafeColourForValue(ColumnColorThresholds.Last(), ColumnValues.Last()));
+				Columns.Add(new ColumnInfo(MvpStatName, fpsChartData.MVP, GetStatColourThresholdList(MvpStatName)));
 			}
 
-			List<bool> ColumnIsAvgValueList = new List<bool>();
-			for (int i = 0; i < ColumnNames.Count; i++)
-			{
-				ColumnIsAvgValueList.Add(false);
-			}
+			// Output CSV stats
+			int csvStatColumnStartIndex = Columns.Count;
 			foreach (string statName in stats)
 			{
-				string[] StatTokens = statName.Split('(');
+				string baseStatName = statName;
+
+				string [] statAttributes=new string[0];
+				int bracketIndex = statName.IndexOf('(');
+				if (bracketIndex != -1)
+				{
+					baseStatName = statName.Substring(0,bracketIndex);
+					int rBracketIndex = statName.LastIndexOf(')');
+					if (rBracketIndex > bracketIndex)
+					{
+						string attributesStr = statName.Substring(bracketIndex + 1, rBracketIndex - (bracketIndex + 1)).ToLower();
+						statAttributes = attributesStr.Split(' ');
+					}
+				}
 
-				float value = 0;
+				float value = 0.0f;
 				string ValueType = " Avg";
 				bool bIsAvg = false;
-				if (!csvStats.Stats.ContainsKey(StatTokens[0].ToLower()))
+				if (!csvStats.Stats.ContainsKey(baseStatName.ToLower()))
 				{
 					continue;
 				}
-				if (StatTokens.Length > 1 && StatTokens[1].ToLower().Contains("min"))
+
+				bool bUnstripped = statAttributes.Contains("unstripped");
+				StatSamples stat = bUnstripped ? csvStatsUnstripped.Stats[baseStatName.ToLower()] : csvStats.Stats[baseStatName.ToLower()];
+				if (statAttributes.Contains("min"))
 				{
-					value = csvStats.Stats[StatTokens[0].ToLower()].ComputeMinValue();
+					value = stat.ComputeMinValue();
 					ValueType = " Min";
 				}
-				else if (StatTokens.Length > 1 && StatTokens[1].ToLower().Contains("max"))
+				else if (statAttributes.Contains("max"))
 				{
-					value = csvStats.Stats[StatTokens[0].ToLower()].ComputeMaxValue();
+					value = stat.ComputeMaxValue();
 					ValueType = " Max";
 				}
 				else
 				{
-					value = csvStats.Stats[StatTokens[0].ToLower()].average;
+					value = stat.average;
 					bIsAvg = true;
 				}
-				ColumnIsAvgValueList.Add(bIsAvg);
-				ColumnNames.Add(StatTokens[0] + ValueType);
-				ColumnValues.Add(value);
-				ColumnColorThresholds.Add(GetStatColourThresholdList(statName));
-				ColumnColors.Add(ColourThresholdList.GetSafeColourForValue(ColumnColorThresholds.Last(), ColumnValues.Last()));
+
+				string detailStr = null;
+				if (bUnstripped)
+				{
+					detailStr = "all frames";
+				}
+				Columns.Add(new ColumnInfo(baseStatName + ValueType, value, GetStatColourThresholdList(statName), detailStr, bIsAvg));
 			}
 
 			// Output summary table row data
 			if (rowData != null)
 			{
-				for (int i = 0; i < ColumnNames.Count; i++)
+				foreach (ColumnInfo column in Columns)
 				{
-					string columnName = ColumnNames[i];
+					string columnName = column.Name;
 
 					// Output simply MVP to rowData instead of MVP30 etc
 					if (columnName.StartsWith("MVP"))
@@ -219,7 +244,7 @@ namespace PerfSummaries
 						columnName = "MVP";
 					}
 					// Hide pre-existing stats with the same name
-					if (ColumnIsAvgValueList[i] && columnName.EndsWith(" Avg"))
+					if (column.isAverageValue && columnName.EndsWith(" Avg"))
 					{
 						string originalStatName = columnName.Substring(0, columnName.Length - 4).ToLower();
 						SummaryTableElement smv;
@@ -231,7 +256,7 @@ namespace PerfSummaries
 							}
 						}
 					}
-					rowData.Add(SummaryTableElement.Type.SummaryTableMetric, columnName, ColumnValues[i], ColumnColorThresholds[i]);
+					rowData.Add(SummaryTableElement.Type.SummaryTableMetric, columnName, column.Value, column.ColorThresholds);
 				}
 				rowData.Add(SummaryTableElement.Type.SummaryTableMetric, "TargetFPS", (double)fps);
 			}
@@ -243,15 +268,21 @@ namespace PerfSummaries
 				string ValueRow = "";
 				HeaderRow += "<th>Section Name</th>";
 				ValueRow += "<td>Entire Run</td>";
-				for (int i = 0; i < ColumnNames.Count; i++)
-				{
-					string columnName = ColumnNames[i];
+				foreach (ColumnInfo column in Columns)
+				{ 
+					string columnName = column.Name;
 					if (columnName.ToLower().EndsWith("time"))
 					{
 						columnName += " (ms)";
 					}
+
+					if (column.Details != null)
+					{
+						columnName += " ("+ column.Details + ")";
+					}
+
 					HeaderRow += "<th>" + TableUtil.FormatStatName(columnName) + "</th>";
-					ValueRow += "<td bgcolor=" + ColumnColors[i] + ">" + ColumnValues[i].ToString("0.00") + "</td>";
+					ValueRow += "<td bgcolor=" + column.Color + ">" + column.Value.ToString("0.00") + "</td>";
 				}
 				htmlFile.WriteLine("  <h2>FPSChart</h2>");
 				htmlFile.WriteLine("<table border='0' style='width:400'>");
@@ -262,6 +293,15 @@ namespace PerfSummaries
 			// Output CSV
 			if (statsCsvFile != null)
 			{
+				List<string> ColumnNames = new List<string>();
+				List<double> ColumnValues = new List<double>();
+				List<string> ColumnColors = new List<string>();
+				foreach (ColumnInfo column in Columns)
+				{
+					ColumnNames.Add(column.Name);
+					ColumnValues.Add(column.Value);
+					ColumnColors.Add(column.Color);
+				}
 				statsCsvFile.Write("Section Name,");
 				statsCsvFile.WriteLine(string.Join(",", ColumnNames));
 
@@ -273,13 +313,19 @@ namespace PerfSummaries
 				statsCsvFile.WriteLine(string.Join(",", ColumnColors));
 			}
 
+
+
 			if (csvStats.Events.Count > 0)
 			{
+				Dictionary<string, ColumnInfo> ColumnDict = new Dictionary<string, ColumnInfo>();
+				foreach (ColumnInfo column in Columns)
+				{
+					ColumnDict[column.Name] = column;
+				}
+
 				// Per-event breakdown
 				foreach (CaptureRange CapRange in captures)
 				{
-					ColumnValues.Clear();
-					ColumnColors.Clear();
 					CaptureData CaptureFrameTimes = GetFramesForCapture(CapRange, frameTimes, csvStats.Events);
 
 					if (CaptureFrameTimes == null)
@@ -293,24 +339,21 @@ namespace PerfSummaries
 						continue;
 					}
 
-					ColumnValues.Add(captureFpsChartData.TotalTimeSeconds);
-					ColumnColors.Add("\'#ffffff\'");
-
-					ColumnValues.Add(captureFpsChartData.HitchesPerMinute);
-					ColumnColors.Add(GetStatThresholdColour("Hitches/Min", captureFpsChartData.HitchesPerMinute));
+					ColumnDict["Total Time (s)"].Value = captureFpsChartData.TotalTimeSeconds;
+					ColumnDict["Hitches/Min"].Value = captureFpsChartData.HitchesPerMinute;
 
 					if (!bIgnoreHitchTimePercent)
 					{
-						ColumnValues.Add(captureFpsChartData.HitchTimePercent);
-						ColumnColors.Add(GetStatThresholdColour("HitchTimePercent", captureFpsChartData.HitchTimePercent));
+						ColumnDict["HitchTimePercent"].Value = captureFpsChartData.HitchTimePercent;
 					}
 
 					if (!bIgnoreMVP)
 					{
-						ColumnValues.Add(captureFpsChartData.MVP);
-						ColumnColors.Add(GetStatThresholdColour("MVP" + fps.ToString(), captureFpsChartData.MVP));
+						ColumnDict[MvpStatName].Value = captureFpsChartData.MVP;
 					}
 
+					// Update the CSV stat values
+					int columnIndex = csvStatColumnStartIndex;
 					foreach (string statName in stats)
 					{
 						string StatToCheck = statName.Split('(')[0];
@@ -334,9 +377,14 @@ namespace PerfSummaries
 						{
 							value = csvStats.Stats[StatTokens[0].ToLower()].ComputeAverage(CaptureFrameTimes.startIndex, CaptureFrameTimes.endIndex);
 						}
+						Columns[columnIndex].Value = value;
+						columnIndex++;
+					}
 
-						ColumnValues.Add(value);
-						ColumnColors.Add(GetStatThresholdColour(statName, value));
+					// Recompute colors
+					foreach (ColumnInfo column in Columns)
+					{
+						column.UpdateColor();
 					}
 
 					// Output HTML
@@ -344,9 +392,9 @@ namespace PerfSummaries
 					{
 						string ValueRow = "";
 						ValueRow += "<td>" + CapRange.name + "</td>";
-						for (int i = 0; i < ColumnNames.Count; i++)
+						foreach (ColumnInfo column in Columns)
 						{
-							ValueRow += "<td bgcolor=" + ColumnColors[i] + ">" + ColumnValues[i].ToString("0.00") + "</td>";
+							ValueRow += "<td bgcolor=" + column.Color + ">" + column.Value.ToString("0.00") + "</td>";
 						}
 						htmlFile.WriteLine("  <tr>" + ValueRow + "</tr>");
 					}
@@ -354,6 +402,14 @@ namespace PerfSummaries
 					// Output CSV
 					if (statsCsvFile != null)
 					{
+						List<double> ColumnValues = new List<double>();
+						List<string> ColumnColors = new List<string>();
+						foreach (ColumnInfo column in Columns)
+						{
+							ColumnValues.Add(column.Value);
+							ColumnColors.Add(column.Color);
+						}
+
 						statsCsvFile.Write(CapRange.name + ",");
 						statsCsvFile.WriteLine(string.Join(",", ColumnValues));
 
diff --git a/Engine/Source/Programs/CSVTools/PerfReportTool/Summaries/HistogramSummary.cs b/Engine/Source/Programs/CSVTools/PerfReportTool/Summaries/HistogramSummary.cs
--- a/Engine/Source/Programs/CSVTools/PerfReportTool/Summaries/HistogramSummary.cs
+++ b/Engine/Source/Programs/CSVTools/PerfReportTool/Summaries/HistogramSummary.cs
@@ -45,7 +45,7 @@ namespace PerfSummaries
 
 		public override string GetName() { return "histogram"; }
 
-		public override void WriteSummaryData(System.IO.StreamWriter htmlFile, CsvStats csvStats, bool bWriteSummaryCsv, SummaryTableRowData metadata, string htmlFileName)
+		public override void WriteSummaryData(System.IO.StreamWriter htmlFile, CsvStats csvStats, CsvStats csvStatsUnstripped, bool bWriteSummaryCsv, SummaryTableRowData metadata, string htmlFileName)
 		{
 			// Only HTML reporting is supported (does not output summary table row data)
 			if (htmlFile == null)
diff --git a/Engine/Source/Programs/CSVTools/PerfReportTool/Summaries/HitchSummary.cs b/Engine/Source/Programs/CSVTools/PerfReportTool/Summaries/HitchSummary.cs
--- a/Engine/Source/Programs/CSVTools/PerfReportTool/Summaries/HitchSummary.cs
+++ b/Engine/Source/Programs/CSVTools/PerfReportTool/Summaries/HitchSummary.cs
@@ -27,7 +27,7 @@ namespace PerfSummaries
 
 		public override string GetName() { return "hitches"; }
 
-		public override void WriteSummaryData(System.IO.StreamWriter htmlFile, CsvStats csvStats, bool bWriteSummaryCsv, SummaryTableRowData summaryTableRowData, string htmlFileName)
+		public override void WriteSummaryData(System.IO.StreamWriter htmlFile, CsvStats csvStats, CsvStats csvStatsUnstripped, bool bWriteSummaryCsv, SummaryTableRowData summaryTableRowData, string htmlFileName)
 		{
 			// if we have 20k frames in a typical flythrough then 20 frames would be red
 			float redThresholdFor50ms = (float)csvStats.SampleCount / 500.0f;
diff --git a/Engine/Source/Programs/CSVTools/PerfReportTool/Summaries/MapOverlaySummary.cs b/Engine/Source/Programs/CSVTools/PerfReportTool/Summaries/MapOverlaySummary.cs
--- a/Engine/Source/Programs/CSVTools/PerfReportTool/Summaries/MapOverlaySummary.cs
+++ b/Engine/Source/Programs/CSVTools/PerfReportTool/Summaries/MapOverlaySummary.cs
@@ -135,7 +135,7 @@ namespace PerfSummaries
 			}
 		}
 
-		public override void WriteSummaryData(System.IO.StreamWriter htmlFile, CsvStats csvStats, bool bWriteSummaryCsv, SummaryTableRowData rowData, string htmlFileName)
+		public override void WriteSummaryData(System.IO.StreamWriter htmlFile, CsvStats csvStats, CsvStats csvStatsUnstripped, bool bWriteSummaryCsv, SummaryTableRowData rowData, string htmlFileName)
 		{
 			// Output HTML
 			if (htmlFile != null)
@@ -143,11 +143,8 @@ namespace PerfSummaries
 				string outputDirectory = System.IO.Path.GetDirectoryName(System.IO.Path.GetFullPath(htmlFileName));
 				string outputMapFilename = System.IO.Path.Combine(outputDirectory, destImageFilename);
 
-				if (File.Exists(outputMapFilename))
-				{
-					Console.WriteLine("Map file already exists. Skipping copy: "+ outputMapFilename);
-				}
-				else
+				// Skip the copy if the output file already exists
+				if (!File.Exists(outputMapFilename))
 				{
 					if (File.Exists(sourceImagePath))
 					{
diff --git a/Engine/Source/Programs/CSVTools/PerfReportTool/Summaries/PeakSummary.cs b/Engine/Source/Programs/CSVTools/PerfReportTool/Summaries/PeakSummary.cs
--- a/Engine/Source/Programs/CSVTools/PerfReportTool/Summaries/PeakSummary.cs
+++ b/Engine/Source/Programs/CSVTools/PerfReportTool/Summaries/PeakSummary.cs
@@ -173,7 +173,7 @@ namespace PerfSummaries
 			return null;
 		}
 
-		public override void WriteSummaryData(System.IO.StreamWriter htmlFile, CsvStats csvStats, bool bWriteSummaryCsv, SummaryTableRowData summaryTableRowData, string htmlFileName)
+		public override void WriteSummaryData(System.IO.StreamWriter htmlFile, CsvStats csvStats, CsvStats csvStatsUnstripped, bool bWriteSummaryCsv, SummaryTableRowData summaryTableRowData, string htmlFileName)
 		{
 			// Only HTML reporting is supported (does not output summary table row data)
 			StreamWriter LLMCsvData = null;
diff --git a/Engine/Source/Programs/CSVTools/PerfReportTool/Summary.cs b/Engine/Source/Programs/CSVTools/PerfReportTool/Summary.cs
--- a/Engine/Source/Programs/CSVTools/PerfReportTool/Summary.cs
+++ b/Engine/Source/Programs/CSVTools/PerfReportTool/Summary.cs
@@ -78,7 +78,7 @@ namespace PerfSummaries
             StatThresholds = new Dictionary<string, ColourThresholdList>();
 
         }
-        public virtual void WriteSummaryData(System.IO.StreamWriter htmlFile, CsvStats csvStats, bool bWriteSummaryCsv, SummaryTableRowData rowData, string htmlFileName)
+        public virtual void WriteSummaryData(System.IO.StreamWriter htmlFile, CsvStats csvStats, CsvStats csvStatsUnstripped, bool bWriteSummaryCsv, SummaryTableRowData rowData, string htmlFileName)
         { }
 
         public virtual void PostInit(ReportTypeInfo reportTypeInfo, CsvStats csvStats)
