CsvToSvg 2.56 - graph rendering optimizations
- Split graph into segments to make re-rendering faster. This makes graphs slightly larger*
- Also make legend more visible

*This is a stop-gap measure to make reports usable. Need to revisit
- Overdraw in stacked graphs is the main issue - need to look at whether we can do this in a less brute-force way

diff --git a/Engine/Source/Programs/CSVTools/CsvToSvgLib/CsvToSvgLib.cs b/Engine/Source/Programs/CSVTools/CsvToSvgLib/CsvToSvgLib.cs
--- a/Engine/Source/Programs/CSVTools/CsvToSvgLib/CsvToSvgLib.cs
+++ b/Engine/Source/Programs/CSVTools/CsvToSvgLib/CsvToSvgLib.cs
@@ -12,7 +12,7 @@
 
 	public class CsvToSvgLibVersion
 	{
-		private static string VersionString = "2.55";
+		private static string VersionString = "2.56";
 
 		public static string Get() { return VersionString; }
 	};
@@ -617,7 +617,6 @@
 
 			// Draw the graphs
 			int csvIndex = 0;
-			uint colourIndex = 0;
 			svg.WriteLine("<g id='graphArea'>");
 			if (graphParams.percentile)
 			{
@@ -631,7 +630,6 @@
 					{
 						string statID = "Stat_" + csvIndex + "_" + stat.Name + "<UNIQUE>";
 						DrawPercentileGraph(svg, stat.samples, stat.colour, graphRect, range, statID);
-						colourIndex++;
 					}
 					csvIndex++;
 				}
@@ -645,11 +643,24 @@
 					DrawEventLines(svg, theme, csvStat.Events, graphRect, range);
 					DrawEventHighlightRegions(svg, csvStat, graphRect, range, graphParams.highlightEventRegions);
 
+					int MaxSegments = 1;
+					if (graphParams.stacked)
+					{
+						MaxSegments = 4;
+						if (csvStat.Stats.Values.Count > 15)
+						{
+							MaxSegments = 8;
+							if (csvStat.Stats.Values.Count > 30)
+							{
+								MaxSegments = 32;
+							}
+						}
+					}
+
 					foreach (StatSamples stat in csvStat.Stats.Values)
 					{
 						string statID = "Stat_" + csvIndex + "_" + stat.Name + "<UNIQUE>";
-						DrawGraph(svg, stat.samples, stat.colour, graphRect, range, thickness, statID, graphParams);
-						colourIndex++;
+						DrawGraph(svg, stat.samples, stat.colour, graphRect, range, thickness, statID, graphParams, MaxSegments);
 					}
 
 					if (graphParams.showEventNameText)
@@ -1402,7 +1413,7 @@
 			return CompressedPoints.ToList();
 		}
 
-		void DrawGraph(SvgFile svg, List<float> samples, Colour colour, Rect rect, Range range, float thickness, string id, GraphParams graphParams)
+		void DrawGraph(SvgFile svg, List<float> samples, Colour colour, Rect rect, Range range, float thickness, string id, GraphParams graphParams, int MaxSegments=1)
 		{
 			int n = Math.Min(samples.Count, (int)(range.MaxX + 0.5));
 			int start = (int)(range.MinX + 0.5);
@@ -1428,6 +1439,14 @@
 
 				string idString = id.Length > 0 ? "id='" + id + "'" : "";
 
+				string fillString = "none";
+				if (graphParams.stacked)
+				{
+					thickness = Math.Min(thickness, 0.01f);
+					colour.alpha = 1.0f;
+					fillString = colour.SVGStringNoQuotes();
+				}
+
 				float graphScale = 1.0f;
 				if (graphParams.bFixedPointGraphs)
 				{
@@ -1435,12 +1454,41 @@
 					// We scale the polyline by the inverse of the scaling that we're applying to the points
 					graphScale = graphParams.fixedPointPrecisionScale;
 					float oneOverScale = 1.0f / graphScale;
-					svg.WriteLine("<polyline " + idString + " transform='scale("+ oneOverScale + ","+ oneOverScale + ")' points='");
-					foreach (Vec2 point in RawPoints)
+
+					// Divide the graph into segments to make it faster to render
+					float originSvgY = ToSvgY(0.0f, rect, range);
+
+					int numSegments = Math.Max( Math.Min(RawPoints.Count / 500, 1), MaxSegments);
+					int pointsPerSegment = RawPoints.Count / numSegments;
+					for ( int i=0;i<numSegments; i++)
 					{
-						int x = (int)Math.Round(point.X * graphScale, 0);
-						int y = (int)Math.Round(point.Y * graphScale, 0);
-						svg.WriteFast(" " + x + "," + y);
+						int startIndex = i * pointsPerSegment;
+						int endIndex = startIndex + pointsPerSegment + 1;
+						if (i == numSegments - 1)
+						{
+							endIndex = RawPoints.Count;
+						}
+
+						string segmentIdString = id.Length > 0 ? "id='" + id + "_Seg-"+i+ "'" : "";
+						svg.WriteLine("<polyline " + segmentIdString + " transform='scale(" + oneOverScale + "," + oneOverScale + ")' points='");
+
+						for (int pointIndex=startIndex ; pointIndex < endIndex; pointIndex++)
+						{
+							Vec2 point = RawPoints[pointIndex];
+							int x = (int)Math.Round(point.X * graphScale, 0);
+							int y = (int)Math.Round(point.Y * graphScale, 0);
+							svg.WriteFast(" " + x + "," + y);
+						}
+
+						if (graphParams.stacked)
+						{
+							int firstX = (int)Math.Round(RawPoints[startIndex].X * graphScale);
+							int lastX = (int)Math.Round(RawPoints[endIndex - 1].X * graphScale);
+							int originY = (int)Math.Round(ToSvgY(0.0f, rect, range) * graphScale);
+							svg.WriteFast(" " + lastX + ","+ originY);
+							svg.WriteFast(" " + firstX + ","+ originY);
+						}
+						svg.WriteLine("' style='fill:" + fillString + ";stroke-width:" + thickness * graphScale + "; clip-path: url(#graphClipRect)' stroke=" + colour.SVGString() + "/>");
 					}
 				}
 				else
@@ -1461,22 +1509,17 @@
 						float y = point.Y;
 						svg.WriteFast(" " + x.ToString(formatStr) + "," + y.ToString(formatStr));
 					}
+
+					if (graphParams.stacked)
+					{
+						float lastSample = samples[n - 1];
+						svg.WriteFast(" " + ToSvgX((float)n + 20, rect, range) * graphScale + "," + ToSvgY(lastSample, rect, range) * graphScale);
+						svg.WriteFast(" " + ToSvgX((float)n + 20, rect, range) * graphScale + "," + ToSvgY(0.0f, rect, range) * graphScale);
+						svg.WriteFast(" " + ToSvgX((float)n - 20, rect, range) * graphScale + "," + ToSvgY(0.0f, rect, range) * graphScale);
+						svg.WriteFast(" " + ToSvgX(start, rect, range) * graphScale + "," + ToSvgY(0.0f, rect, range) * graphScale );
+					}
+					svg.WriteLine("' style='fill:" + fillString + ";stroke-width:" + thickness * graphScale + "; clip-path: url(#graphClipRect)' stroke=" + colour.SVGString() + "/>");
 				}
-
-				string fillString = "none";
-				if (graphParams.stacked)
-				{
-					thickness = Math.Min(thickness, 0.01f);
-					float lastSample = samples[n - 1];
-					svg.WriteFast(" " + ToSvgX((float)n + 20, rect, range) * graphScale + "," + ToSvgY(lastSample, rect, range) * graphScale);
-					svg.WriteFast(" " + ToSvgX((float)n + 20, rect, range) * graphScale + "," + ToSvgY(0.0f, rect, range) * graphScale);
-					svg.WriteFast(" " + ToSvgX((float)n - 20, rect, range) * graphScale + "," + ToSvgY(0.0f, rect, range) * graphScale);
-					svg.WriteFast(" " + ToSvgX(start, rect, range) * graphScale + "," + ToSvgY(0.0f, rect, range) * graphScale );
-
-					colour.alpha = 1.0f;
-					fillString = colour.SVGStringNoQuotes();
-				}
-				svg.WriteLine("' style='fill:" + fillString + ";stroke-width:" + thickness * graphScale + "; clip-path: url(#graphClipRect)' stroke=" + colour.SVGString() + "/>");
 			}
 		}
 		void DrawPercentileGraph(SvgFile svg, List<float> samples, Colour colour, Rect rect, Range range, string id)
@@ -1747,7 +1790,8 @@
 			}
 			float y = 30; // rect.height - 25;
 
-			svg.WriteLine("<g id='LegendPanel<UNIQUE>'>");// transform='translate(5) rotate(45 50 50)'> ");
+			svg.WriteLine("<g id='LegendPanel<UNIQUE>' fill='rgba(255,0,0,0.3)'>");// transform='translate(5) rotate(45 50 50)'> ");
+			svg.WriteLine("<rect x='"+x+"' y='"+y+"' width='0' height='100' fill='rgba(0,0,0,0.3)' blend='1' id='legendPanelRect<UNIQUE>' rx='5' ry='5'/>");
 
 			// Check if the total is a fraction
 			bool legendValueIsWholeNumber = false;
@@ -1780,7 +1824,7 @@
 						{
 							formatString = "0";
 						}
-						svg.WriteLine("<text x='" + (x + 15) + "' y='" + y + "' fill=" + theme.MediumTextColour.SVGString() + " font-size='10' font-family='Helvetica' style='text-anchor: start' filter='url(#dropShadowFilter)'>" + legendValue.ToString(formatString) + "</text>");
+						svg.WriteLine("<text x='" + (x + 15) + "' y='" + y + "' fill=" + theme.TextColour.SVGString() + " font-size='10' font-family='Helvetica' style='text-anchor: start' filter='url(#dropShadowFilter)'>" + legendValue.ToString(formatString) + "</text>");
 					}
 
 					y += 16;
@@ -1790,7 +1834,7 @@
 			if (graphParams.showTotals || graphParams.showAverages)
 			{
 				string legendValueTypeString = graphParams.showTotals ? "(sum)" : "(avg)";
-				svg.WriteLine("<text x='" + (x + 15) + "' y='" + y + "' fill=" + theme.MediumTextColour.SVGString() + " font-size='10' font-family='Helvetica' style='text-anchor: start' filter='url(#dropShadowFilter)'>" + legendValueTypeString + "</text>");
+				svg.WriteLine("<text x='" + (x + 15) + "' y='" + y + "' fill=" + theme.TextColour.SVGString() + " font-size='10' font-family='Helvetica' style='text-anchor: start' filter='url(#dropShadowFilter)'>" + legendValueTypeString + "</text>");
 			}
 
 			svg.WriteLine("</g>");
@@ -1971,17 +2015,12 @@
 			}
 
 
-			//style='text-anchor: "
-			//svg.WriteLine("<g transform='translate(0,-24)'>");
-			svg.WriteLine("<g id='interactivePanel<UNIQUE>' visibility='hidden'>");
-			DrawVerticalLine(svg, 0, new Colour(255, 255, 255, 0.4f), rect, range, 1.0f, true, true, "", true);
-			svg.WriteLine("<g id='interactivePanelInnerWithRect<UNIQUE>'>");
-			svg.WriteLine("<rect x='0' y='0' width='100' height='100' fill='rgba(0,0,0,0.3)' blend='1' id='interactivePanelRect<UNIQUE>' rx='5' ry='5'/>");
-			svg.WriteLine("<g id='interactivePanelInner<UNIQUE>'>");
+			// Create a hidden panel for storing all the stat group elements
+			svg.WriteLine("<g id='interactivePanelInnerHidden<UNIQUE>' visibility='collapse'>");
 			Rect zeroRect = new Rect(0, 0, 0, 0);
 			foreach (InteractiveStatInfo statInfo in interactiveStats)
 			{
-				svg.WriteLine("<g id='" + statInfo.jsGroupElementId + "' transform='translate(0,0)' visibility='hidden'>");
+				svg.WriteLine("<g id='" + statInfo.jsGroupElementId + "' transform='translate(0,0)' visibility='inherit'>");
 				int textXOffset = -5;
 				if (statInfo.colour != null)
 				{
@@ -1993,6 +2032,15 @@
 
 				svg.WriteLine("</g>");
 			}
+			svg.WriteLine("</g>"); // interactivePanelInnerHidden
+			svg.WriteLine("<g id='interactivePanel<UNIQUE>' visibility='hidden'>");
+			DrawVerticalLine(svg, 0, new Colour(255, 255, 255, 0.4f), rect, range, 1.0f, true, true, "", true);
+			svg.WriteLine("<g id='interactivePanelInnerWithRect<UNIQUE>'>");
+			svg.WriteLine("<rect x='0' y='0' width='100' height='100' fill='rgba(0,0,0,0.3)' blend='1' id='interactivePanelRect<UNIQUE>' rx='5' ry='5'/>");
+
+
+			svg.WriteLine("<g id='interactivePanelInner<UNIQUE>'>");
+
 			svg.WriteLine("</g>"); // interactivePanelInner
 			svg.WriteLine("</g>"); // interactivePanelInnerWithRect
 			svg.WriteLine("</g>"); // interactivepanel
@@ -2036,6 +2084,13 @@
 
 			svg.WriteLine("function OnLoaded<UNIQUE>(evt)");
 			svg.WriteLine("{");
+			svg.WriteLine("    var legendPanel = document.getElementById('LegendPanel<UNIQUE>');");
+			svg.WriteLine("    var legendRect = document.getElementById('legendPanelRect<UNIQUE>');");
+			svg.WriteLine("    var legendBbox = legendPanel.getBBox();");
+			svg.WriteLine("    legendRect.setAttribute('width',legendBbox.width+80);");
+			svg.WriteLine("    legendRect.setAttribute('height',legendBbox.height+8);");
+			svg.WriteLine("    legendRect.setAttribute('x',legendRect.getAttribute('x')-legendBbox.width);");
+			//svg.WriteLine("    legendRect.setAttribute('y',legendBbox.y-8);");
 			svg.WriteLine("}");
 
 			if (bSnapToPeaks)
@@ -2171,32 +2226,43 @@
 				svg.WriteLine(", groupElement: " + groupElementString + ", textElement: " + textElementString + " },");
 			}
 			svg.WriteLine("    ];");
-
 			svg.WriteLine("    samples.sort(compareSamples);");
+
+			// Draw the interactive legend
+			// Move all elements from the visible rect into the hidden rect
+			svg.WriteLine("    var panelInner = document.getElementById('interactivePanelInner<UNIQUE>');");
+			svg.WriteLine("    var panelInnerHidden = document.getElementById('interactivePanelInnerHidden<UNIQUE>');");
+			svg.WriteLine("    while (panelInner.childNodes.length>0)");
+			svg.WriteLine("    {");
+			svg.WriteLine("        panelInnerHidden.appendChild(panelInner.childNodes[0]);");
+			svg.WriteLine("    }");
+
+			svg.WriteLine("    var maxSampleCountToDisplay="+(rect.height/15)+";");
 			svg.WriteLine("    for ( i=0;i<samples.length;i++ )");
 			svg.WriteLine("    {");
 			svg.WriteLine("        var groupElement = samples[i].groupElement;");
-			svg.WriteLine("        if ( samples[i].value > 0 )");
+			svg.WriteLine("        if ( samples[i].value > 0 && i<=maxSampleCountToDisplay)");
 			svg.WriteLine("        {");
-			svg.WriteLine("            groupElement.setAttribute('visibility','inherit');");
 			svg.WriteLine("            groupElement.setAttribute('transform','translate(0,'+textY+')');");
 			svg.WriteLine("            var textElement = samples[i].textElement;");
-			//            svg.WriteLine("            var textNode = document.createTextNode(interactivePanelInner.getBBox().x);");
 			svg.WriteLine("            var textNode = document.createTextNode(samples[i].value);");
 			svg.WriteLine("            textElement.replaceChild(textNode,textElement.childNodes[0]); ");
+			svg.WriteLine("            if (groupElement.parentNode != panelInner)");
+			svg.WriteLine("              panelInner.appendChild(groupElement);");
 			svg.WriteLine("            textY += 15;");
 			svg.WriteLine("        }");
 			svg.WriteLine("        else");
 			svg.WriteLine("        {");
-			svg.WriteLine("            groupElement.setAttribute('visibility','hidden')");
+			svg.WriteLine("            if (groupElement.parentNode != panelInnerHidden)");
+			svg.WriteLine("                panelInnerHidden.appendChild(groupElement);");
 			svg.WriteLine("        }");
 			svg.WriteLine("    }");
 
 			svg.WriteLine("    var panelRect = document.getElementById('interactivePanelRect<UNIQUE>');");
 			svg.WriteLine("    var bbox = interactivePanelInner.getBBox();");
-			svg.WriteLine("    panelRect.setAttribute('width',bbox.width+32);");
+			svg.WriteLine("    panelRect.setAttribute('width',bbox.width+16);");
 			svg.WriteLine("    panelRect.setAttribute('height',textY+8);");
-			svg.WriteLine("    panelRect.setAttribute('x',bbox.x-16);");
+			svg.WriteLine("    panelRect.setAttribute('x',bbox.x-8);");
 			svg.WriteLine("    panelRect.setAttribute('y',bbox.y-8);");
 
 			float maxX = rect.x + rect.width;
