PerfReportTool 4.77 - Use InvariantCulture to fix graph generation issues with some cultures

diff --git a/Engine/Source/Programs/CSVTools/PerfReportTool/PerfReportTool.cs b/Engine/Source/Programs/CSVTools/PerfReportTool/PerfReportTool.cs
--- a/Engine/Source/Programs/CSVTools/PerfReportTool/PerfReportTool.cs
+++ b/Engine/Source/Programs/CSVTools/PerfReportTool/PerfReportTool.cs
@@ -22,7 +22,7 @@ namespace PerfReportTool
 {
     class Version
     {
-        private static string VersionString = "4.76";
+        private static string VersionString = "4.77";
 
         public static string Get() { return VersionString; }
     };
@@ -201,8 +201,10 @@ namespace PerfReportTool
 
 		void Run(string[] args)
 		{
-            // Read the command line
-            if (args.Length < 1)
+			System.Globalization.CultureInfo.DefaultThreadCurrentCulture = System.Globalization.CultureInfo.InvariantCulture;
+
+			// Read the command line
+			if (args.Length < 1)
 			{
 				WriteLine("Invalid args");
 				WriteLine(formatString);
