CsvToSvgLib 2.57
Make the smoothing downsample multiplier take into account the samples per pixel ratio. We should never have a downsample kernel which is coarser than the visible res.

diff --git a/Engine/Source/Programs/CSVTools/CsvToSvgLib/CsvToSvgLib.cs b/Engine/Source/Programs/CSVTools/CsvToSvgLib/CsvToSvgLib.cs
--- a/Engine/Source/Programs/CSVTools/CsvToSvgLib/CsvToSvgLib.cs
+++ b/Engine/Source/Programs/CSVTools/CsvToSvgLib/CsvToSvgLib.cs
@@ -12,7 +12,7 @@
 
 	public class CsvToSvgLibVersion
 	{
-		private static string VersionString = "2.56";
+		private static string VersionString = "2.57";
 
 		public static string Get() { return VersionString; }
 	};
@@ -514,7 +514,7 @@
 						float kernelSizeF = percent * 0.01f * (float)csvStatsList[i].SampleCount + 0.5f;
 						kernelSize = (int)kernelSizeF;
 					}
-					csvStatsList[i] = SmoothStats(csvStatsList[i], kernelSize, graphParams.bSmoothMultithreaded);
+					csvStatsList[i] = SmoothStats(csvStatsList[i], kernelSize, graphParams);
 				}
 				perfLog.LogTiming("SmoothStats");
 			}
@@ -786,9 +786,10 @@
 
 		class SmoothKernel
 		{
-			public SmoothKernel(int inKernelSize)
+			public SmoothKernel(int inKernelSize, float maxDownsampleMultiplier = 4.0f)
 			{
-				downsampleLevel = Math.Max( Math.Min( (int)Math.Log((double)inKernelSize / 20.0, 2.0), 4), 0);
+				int maxDownsampleLevel = Math.Min(4, Math.Max(0,(int)Math.Log(maxDownsampleMultiplier, 2.0f)));
+				downsampleLevel = Math.Max( Math.Min( (int)Math.Log((double)inKernelSize / 20.0, 2.0), maxDownsampleLevel), 0);
 
 				kernelSize = inKernelSize >> downsampleLevel;
 				weights = new float[kernelSize];
@@ -885,16 +886,30 @@
 
 
 
-		CsvStats SmoothStats(CsvStats stats, int KernelSize, bool bMultiThreaded)
+		CsvStats SmoothStats(CsvStats stats, int KernelSize, GraphParams graphParams)
 		{
+			bool bMultithreaded = graphParams.bSmoothMultithreaded;
+
+			// Compute the displayed sample count
+			int minX = Math.Max( (int)graphParams.minX, 0 );
+			int maxX = stats.SampleCount;
+			if (graphParams.maxX > 0)
+			{
+				maxX = Math.Min((int)graphParams.maxX, stats.SampleCount);
+			}
+			int sampleCount = maxX - minX;
+
+			// Compute a max downsample level based on the sample count
+			float maxDownsampleMultiplier = (float)sampleCount/graphParams.width;
+
 			// Compute Gaussian Weights
-			SmoothKernel kernel = new SmoothKernel(KernelSize);
+			SmoothKernel kernel = new SmoothKernel(KernelSize, maxDownsampleMultiplier);
 
 			List<StatSamples> statSampleList = stats.Stats.Values.ToList();
 			// Add the stats to smoothstats before the parallel for, so we can preserve the order
 
 			StatSamples[] smoothSamplesArray = new StatSamples[stats.Stats.Count];
-			if (bMultiThreaded)
+			if (bMultithreaded)
 			{
 				int numThreads = Environment.ProcessorCount/2;
 				Parallel.For(0, stats.Stats.Values.Count, new ParallelOptions { MaxDegreeOfParallelism = numThreads }, i =>
