Fix PerfReportFort missing dependencies (csproj was missing from a previous submit)

diff --git a/Engine/Source/Programs/CSVTools/PerfReportTool/PerfreportTool.csproj b/Engine/Source/Programs/CSVTools/PerfReportTool/PerfreportTool.csproj
--- a/Engine/Source/Programs/CSVTools/PerfReportTool/PerfreportTool.csproj
+++ b/Engine/Source/Programs/CSVTools/PerfReportTool/PerfreportTool.csproj
@@ -35,6 +35,7 @@
   <ItemGroup>
     <Reference Include="System" />
     <Reference Include="System.Core" />
+    <Reference Include="System.Drawing" />
     <Reference Include="System.Xml.Linq" />
     <Reference Include="System.Data.DataSetExtensions" />
     <Reference Include="Microsoft.CSharp" />
@@ -70,6 +71,10 @@
       <Project>{75c595cc-83ba-4ffd-a772-17b2343d1061}</Project>
       <Name>CSVStats</Name>
     </ProjectReference>
+    <ProjectReference Include="..\CsvToSvgLib\CsvToSvgLib.csproj">
+      <Project>{75c595cc-83ba-4ffd-abcd-17b2343d1061}</Project>
+      <Name>CsvToSvgLib</Name>
+    </ProjectReference>
   </ItemGroup>
   <Import Project="$(MSBuildToolsPath)\Microsoft.CSharp.targets" />
   <PropertyGroup>
