Upgrade CSV tools to .NET framework 4.8

diff --git a/Engine/Source/Programs/CSVTools/CSVCollate/App.config b/Engine/Source/Programs/CSVTools/CSVCollate/App.config
--- a/Engine/Source/Programs/CSVTools/CSVCollate/App.config
+++ b/Engine/Source/Programs/CSVTools/CSVCollate/App.config
@@ -1,6 +1,6 @@
-﻿<?xml version="1.0" encoding="utf-8" ?>
+<?xml version="1.0" encoding="utf-8"?>
 <configuration>
     <startup> 
-        <supportedRuntime version="v4.0" sku=".NETFramework,Version=v4.5" />
+        <supportedRuntime version="v4.0" sku=".NETFramework,Version=v4.8"/>
     </startup>
-</configuration>
\ No newline at end of file
+</configuration>
diff --git a/Engine/Source/Programs/CSVTools/CSVCollate/CSVCollate.csproj b/Engine/Source/Programs/CSVTools/CSVCollate/CSVCollate.csproj
--- a/Engine/Source/Programs/CSVTools/CSVCollate/CSVCollate.csproj
+++ b/Engine/Source/Programs/CSVTools/CSVCollate/CSVCollate.csproj
@@ -1,5 +1,5 @@
 ﻿<?xml version="1.0" encoding="utf-8"?>
-<Project ToolsVersion="4.0" DefaultTargets="Build" xmlns="http://schemas.microsoft.com/developer/msbuild/2003">
+<Project ToolsVersion="12.0" DefaultTargets="Build" xmlns="http://schemas.microsoft.com/developer/msbuild/2003">
   <Import Project="$(MSBuildExtensionsPath)\$(MSBuildToolsVersion)\Microsoft.Common.props" Condition="Exists('$(MSBuildExtensionsPath)\$(MSBuildToolsVersion)\Microsoft.Common.props')" />
   <PropertyGroup>
     <Configuration Condition=" '$(Configuration)' == '' ">Debug</Configuration>
@@ -9,8 +9,9 @@
     <AppDesignerFolder>Properties</AppDesignerFolder>
     <RootNamespace>CSVCollate</RootNamespace>
     <AssemblyName>CSVCollate</AssemblyName>
-    <TargetFrameworkVersion>v4.5</TargetFrameworkVersion>
+    <TargetFrameworkVersion>v4.8</TargetFrameworkVersion>
     <FileAlignment>512</FileAlignment>
+    <TargetFrameworkProfile />
   </PropertyGroup>
   <PropertyGroup Condition=" '$(Configuration)|$(Platform)' == 'Debug|AnyCPU' ">
     <PlatformTarget>x64</PlatformTarget>
diff --git a/Engine/Source/Programs/CSVTools/CSVFilter/App.config b/Engine/Source/Programs/CSVTools/CSVFilter/App.config
--- a/Engine/Source/Programs/CSVTools/CSVFilter/App.config
+++ b/Engine/Source/Programs/CSVTools/CSVFilter/App.config
@@ -1,6 +1,6 @@
-﻿<?xml version="1.0" encoding="utf-8" ?>
+<?xml version="1.0" encoding="utf-8"?>
 <configuration>
     <startup> 
-        <supportedRuntime version="v4.0" sku=".NETFramework,Version=v4.5" />
+        <supportedRuntime version="v4.0" sku=".NETFramework,Version=v4.8"/>
     </startup>
-</configuration>
\ No newline at end of file
+</configuration>
diff --git a/Engine/Source/Programs/CSVTools/CSVFilter/CSVFilter.csproj b/Engine/Source/Programs/CSVTools/CSVFilter/CSVFilter.csproj
--- a/Engine/Source/Programs/CSVTools/CSVFilter/CSVFilter.csproj
+++ b/Engine/Source/Programs/CSVTools/CSVFilter/CSVFilter.csproj
@@ -9,8 +9,9 @@
     <AppDesignerFolder>Properties</AppDesignerFolder>
     <RootNamespace>CSVFilter</RootNamespace>
     <AssemblyName>CSVFilter</AssemblyName>
-    <TargetFrameworkVersion>v4.5</TargetFrameworkVersion>
+    <TargetFrameworkVersion>v4.8</TargetFrameworkVersion>
     <FileAlignment>512</FileAlignment>
+    <TargetFrameworkProfile />
   </PropertyGroup>
   <PropertyGroup Condition=" '$(Configuration)|$(Platform)' == 'Debug|AnyCPU' ">
     <PlatformTarget>x64</PlatformTarget>
diff --git a/Engine/Source/Programs/CSVTools/CSVSplit/App.config b/Engine/Source/Programs/CSVTools/CSVSplit/App.config
--- a/Engine/Source/Programs/CSVTools/CSVSplit/App.config
+++ b/Engine/Source/Programs/CSVTools/CSVSplit/App.config
@@ -1,6 +1,6 @@
-﻿<?xml version="1.0" encoding="utf-8" ?>
+<?xml version="1.0" encoding="utf-8"?>
 <configuration>
     <startup> 
-        <supportedRuntime version="v4.0" sku=".NETFramework,Version=v4.5" />
+        <supportedRuntime version="v4.0" sku=".NETFramework,Version=v4.8"/>
     </startup>
-</configuration>
\ No newline at end of file
+</configuration>
diff --git a/Engine/Source/Programs/CSVTools/CSVSplit/CSVSplit.csproj b/Engine/Source/Programs/CSVTools/CSVSplit/CSVSplit.csproj
--- a/Engine/Source/Programs/CSVTools/CSVSplit/CSVSplit.csproj
+++ b/Engine/Source/Programs/CSVTools/CSVSplit/CSVSplit.csproj
@@ -9,8 +9,9 @@
     <AppDesignerFolder>Properties</AppDesignerFolder>
     <RootNamespace>CSVSplit</RootNamespace>
     <AssemblyName>CSVSplit</AssemblyName>
-    <TargetFrameworkVersion>v4.5</TargetFrameworkVersion>
+    <TargetFrameworkVersion>v4.8</TargetFrameworkVersion>
     <FileAlignment>512</FileAlignment>
+    <TargetFrameworkProfile />
   </PropertyGroup>
   <PropertyGroup Condition=" '$(Configuration)|$(Platform)' == 'Debug|AnyCPU' ">
     <PlatformTarget>x64</PlatformTarget>
diff --git a/Engine/Source/Programs/CSVTools/CSVToSVG/App.config b/Engine/Source/Programs/CSVTools/CSVToSVG/App.config
--- a/Engine/Source/Programs/CSVTools/CSVToSVG/App.config
+++ b/Engine/Source/Programs/CSVTools/CSVToSVG/App.config
@@ -1,6 +1,6 @@
-﻿<?xml version="1.0" encoding="utf-8" ?>
+<?xml version="1.0" encoding="utf-8"?>
 <configuration>
     <startup> 
-        <supportedRuntime version="v4.0" sku=".NETFramework,Version=v4.5" />
+        <supportedRuntime version="v4.0" sku=".NETFramework,Version=v4.8"/>
     </startup>
-</configuration>
\ No newline at end of file
+</configuration>
diff --git a/Engine/Source/Programs/CSVTools/CSVToSVG/CSVToSVG.csproj b/Engine/Source/Programs/CSVTools/CSVToSVG/CSVToSVG.csproj
--- a/Engine/Source/Programs/CSVTools/CSVToSVG/CSVToSVG.csproj
+++ b/Engine/Source/Programs/CSVTools/CSVToSVG/CSVToSVG.csproj
@@ -1,5 +1,5 @@
 ﻿<?xml version="1.0" encoding="utf-8"?>
-<Project ToolsVersion="4.0" DefaultTargets="Build" xmlns="http://schemas.microsoft.com/developer/msbuild/2003">
+<Project ToolsVersion="12.0" DefaultTargets="Build" xmlns="http://schemas.microsoft.com/developer/msbuild/2003">
   <Import Project="$(MSBuildExtensionsPath)\$(MSBuildToolsVersion)\Microsoft.Common.props" Condition="Exists('$(MSBuildExtensionsPath)\$(MSBuildToolsVersion)\Microsoft.Common.props')" />
   <PropertyGroup>
     <Configuration Condition=" '$(Configuration)' == '' ">Debug</Configuration>
@@ -9,8 +9,9 @@
     <AppDesignerFolder>Properties</AppDesignerFolder>
     <RootNamespace>CSVToSVG</RootNamespace>
     <AssemblyName>CSVToSVG</AssemblyName>
-    <TargetFrameworkVersion>v4.5</TargetFrameworkVersion>
+    <TargetFrameworkVersion>v4.8</TargetFrameworkVersion>
     <FileAlignment>512</FileAlignment>
+    <TargetFrameworkProfile />
   </PropertyGroup>
   <PropertyGroup Condition=" '$(Configuration)|$(Platform)' == 'Debug|AnyCPU' ">
     <PlatformTarget>x64</PlatformTarget>
diff --git a/Engine/Source/Programs/CSVTools/CsvConvert/App.config b/Engine/Source/Programs/CSVTools/CsvConvert/App.config
--- a/Engine/Source/Programs/CSVTools/CsvConvert/App.config
+++ b/Engine/Source/Programs/CSVTools/CsvConvert/App.config
@@ -1,6 +1,6 @@
 ﻿<?xml version="1.0" encoding="utf-8"?>
 <configuration>
     <startup> 
-        <supportedRuntime version="v4.0" sku=".NETFramework,Version=v4.5"/>
+        <supportedRuntime version="v4.0" sku=".NETFramework,Version=v4.8"/>
     </startup>
 </configuration>
diff --git a/Engine/Source/Programs/CSVTools/CsvConvert/CsvConvert.csproj b/Engine/Source/Programs/CSVTools/CsvConvert/CsvConvert.csproj
--- a/Engine/Source/Programs/CSVTools/CsvConvert/CsvConvert.csproj
+++ b/Engine/Source/Programs/CSVTools/CsvConvert/CsvConvert.csproj
@@ -9,7 +9,7 @@
     <AppDesignerFolder>Properties</AppDesignerFolder>
     <RootNamespace>CsvConvert</RootNamespace>
     <AssemblyName>CsvConvert</AssemblyName>
-    <TargetFrameworkVersion>v4.5</TargetFrameworkVersion>
+    <TargetFrameworkVersion>v4.8</TargetFrameworkVersion>
     <FileAlignment>512</FileAlignment>
     <AutoGenerateBindingRedirects>true</AutoGenerateBindingRedirects>
     <TargetFrameworkProfile />
diff --git a/Engine/Source/Programs/CSVTools/CsvInfo/App.config b/Engine/Source/Programs/CSVTools/CsvInfo/App.config
--- a/Engine/Source/Programs/CSVTools/CsvInfo/App.config
+++ b/Engine/Source/Programs/CSVTools/CsvInfo/App.config
@@ -1,6 +1,6 @@
 ﻿<?xml version="1.0" encoding="utf-8"?>
 <configuration>
     <startup> 
-        <supportedRuntime version="v4.0" sku=".NETFramework,Version=v4.5"/>
+        <supportedRuntime version="v4.0" sku=".NETFramework,Version=v4.8"/>
     </startup>
 </configuration>
diff --git a/Engine/Source/Programs/CSVTools/CsvInfo/CsvInfo.csproj b/Engine/Source/Programs/CSVTools/CsvInfo/CsvInfo.csproj
--- a/Engine/Source/Programs/CSVTools/CsvInfo/CsvInfo.csproj
+++ b/Engine/Source/Programs/CSVTools/CsvInfo/CsvInfo.csproj
@@ -9,7 +9,7 @@
     <AppDesignerFolder>Properties</AppDesignerFolder>
     <RootNamespace>CsvInfo</RootNamespace>
     <AssemblyName>CsvInfo</AssemblyName>
-    <TargetFrameworkVersion>v4.5</TargetFrameworkVersion>
+    <TargetFrameworkVersion>v4.8</TargetFrameworkVersion>
     <FileAlignment>512</FileAlignment>
     <AutoGenerateBindingRedirects>true</AutoGenerateBindingRedirects>
     <TargetFrameworkProfile />
diff --git a/Engine/Source/Programs/CSVTools/CsvStats/CsvStats.csproj b/Engine/Source/Programs/CSVTools/CsvStats/CsvStats.csproj
--- a/Engine/Source/Programs/CSVTools/CsvStats/CsvStats.csproj
+++ b/Engine/Source/Programs/CSVTools/CsvStats/CsvStats.csproj
@@ -9,8 +9,9 @@
     <AppDesignerFolder>Properties</AppDesignerFolder>
     <RootNamespace>CsvStats</RootNamespace>
     <AssemblyName>CsvStats</AssemblyName>
-    <TargetFrameworkVersion>v4.5</TargetFrameworkVersion>
+    <TargetFrameworkVersion>v4.8</TargetFrameworkVersion>
     <FileAlignment>512</FileAlignment>
+    <TargetFrameworkProfile />
   </PropertyGroup>
   <PropertyGroup Condition=" '$(Configuration)|$(Platform)' == 'Debug|AnyCPU' ">
     <DebugSymbols>true</DebugSymbols>
diff --git a/Engine/Source/Programs/CSVTools/CsvToSvgLib/CsvToSvgLib.csproj b/Engine/Source/Programs/CSVTools/CsvToSvgLib/CsvToSvgLib.csproj
--- a/Engine/Source/Programs/CSVTools/CsvToSvgLib/CsvToSvgLib.csproj
+++ b/Engine/Source/Programs/CSVTools/CsvToSvgLib/CsvToSvgLib.csproj
@@ -9,8 +9,9 @@
     <AppDesignerFolder>Properties</AppDesignerFolder>
     <RootNamespace>CsvStats</RootNamespace>
     <AssemblyName>CsvToSvgLib</AssemblyName>
-    <TargetFrameworkVersion>v4.5</TargetFrameworkVersion>
+    <TargetFrameworkVersion>v4.8</TargetFrameworkVersion>
     <FileAlignment>512</FileAlignment>
+    <TargetFrameworkProfile />
   </PropertyGroup>
   <PropertyGroup Condition=" '$(Configuration)|$(Platform)' == 'Debug|AnyCPU' ">
     <DebugSymbols>true</DebugSymbols>
diff --git a/Engine/Source/Programs/CSVTools/PerfReportTool/App.config b/Engine/Source/Programs/CSVTools/PerfReportTool/App.config
--- a/Engine/Source/Programs/CSVTools/PerfReportTool/App.config
+++ b/Engine/Source/Programs/CSVTools/PerfReportTool/App.config
@@ -1,9 +1,9 @@
-<?xml version="1.0" encoding="utf-8" ?>
+<?xml version="1.0" encoding="utf-8"?>
 <configuration>
     <startup> 
-        <supportedRuntime version="v4.0" sku=".NETFramework,Version=v4.5" />
+        <supportedRuntime version="v4.0" sku=".NETFramework,Version=v4.8"/>
     </startup>
   <runtime>
-    <gcConcurrent enabled="false" />
+    <gcConcurrent enabled="false"/>
   </runtime>
-</configuration>
\ No newline at end of file
+</configuration>
diff --git a/Engine/Source/Programs/CSVTools/PerfReportTool/PerfreportTool.csproj b/Engine/Source/Programs/CSVTools/PerfReportTool/PerfreportTool.csproj
--- a/Engine/Source/Programs/CSVTools/PerfReportTool/PerfreportTool.csproj
+++ b/Engine/Source/Programs/CSVTools/PerfReportTool/PerfreportTool.csproj
@@ -1,5 +1,5 @@
 ﻿<?xml version="1.0" encoding="utf-8"?>
-<Project ToolsVersion="4.0" DefaultTargets="Build" xmlns="http://schemas.microsoft.com/developer/msbuild/2003">
+<Project ToolsVersion="12.0" DefaultTargets="Build" xmlns="http://schemas.microsoft.com/developer/msbuild/2003">
   <Import Project="$(MSBuildExtensionsPath)\$(MSBuildToolsVersion)\Microsoft.Common.props" Condition="Exists('$(MSBuildExtensionsPath)\$(MSBuildToolsVersion)\Microsoft.Common.props')" />
   <PropertyGroup>
     <Configuration Condition=" '$(Configuration)' == '' ">Debug</Configuration>
@@ -9,8 +9,9 @@
     <AppDesignerFolder>Properties</AppDesignerFolder>
     <RootNamespace>PerfReportTool</RootNamespace>
     <AssemblyName>PerfReportTool</AssemblyName>
-    <TargetFrameworkVersion>v4.5</TargetFrameworkVersion>
+    <TargetFrameworkVersion>v4.8</TargetFrameworkVersion>
     <FileAlignment>512</FileAlignment>
+    <TargetFrameworkProfile />
   </PropertyGroup>
   <PropertyGroup Condition=" '$(Configuration)|$(Platform)' == 'Debug|AnyCPU' ">
     <PlatformTarget>x64</PlatformTarget>
