PerfReportTool 4.78 - add an AutoColorize XML option for summary tables and make this behavior independent of ScrollableFormatting.
- This can be force disabled via the commandline with -autoColorize 0, or enabled with -autoColorize or -autoColorize 1
- The -reverseTable and -scrollableTable params are now also possible to force disable by specifying 0 (using a new GetOptionalBool mechanism)

diff --git a/Engine/Source/Programs/CSVTools/CsvStats/CommandLineTool.cs b/Engine/Source/Programs/CSVTools/CsvStats/CommandLineTool.cs
--- a/Engine/Source/Programs/CSVTools/CsvStats/CommandLineTool.cs
+++ b/Engine/Source/Programs/CSVTools/CsvStats/CommandLineTool.cs
@@ -154,6 +154,15 @@ namespace CSVStats
 			return CommandLineArgs.ContainsKey(key.ToLower());
 		}
 
+		public bool? GetOptionalBoolArg(string key)
+		{
+			if ( CommandLineArgs.ContainsKey(key.ToLower()) )
+			{
+				return CommandLineArgs[key.ToLower()] == "1";
+			}
+			return null;
+		}
+
 		public string GetArg(string key, string defaultValue)
 		{
 			string lowerKey = key.ToLower();
@@ -241,6 +250,12 @@ namespace CSVStats
 			return commandLine.GetBoolArg(key);
         }
 
+		protected bool? GetOptionalBoolArg(string key)
+		{
+			return commandLine.GetOptionalBoolArg(key);
+		}
+
+
 		protected string GetArg(string key, string defaultValue)
 		{
 			return commandLine.GetArg(key, defaultValue);
diff --git a/Engine/Source/Programs/CSVTools/PerfReportTool/PerfReportTool.cs b/Engine/Source/Programs/CSVTools/PerfReportTool/PerfReportTool.cs
--- a/Engine/Source/Programs/CSVTools/PerfReportTool/PerfReportTool.cs
+++ b/Engine/Source/Programs/CSVTools/PerfReportTool/PerfReportTool.cs
@@ -22,7 +22,7 @@ namespace PerfReportTool
 {
     class Version
     {
-        private static string VersionString = "4.77";
+        private static string VersionString = "4.78";
 
         public static string Get() { return VersionString; }
     };
@@ -141,8 +141,9 @@ namespace PerfReportTool
 			"       -externalGraphs : enables external graphs (off by default)\n" +
 			"       -spreadsheetfriendly: outputs a single quote before non-numeric entries in summary tables\n" +
 			"       -noSummaryMinMax: don't make min/max columns for each stat in a condensed summary\n" +
-			"       -reverseTable: Reverses the order of summary tables\n"+
-			"       -scrollableTable: makes the summary table scrollable, with frozen first rows and columns\n" +
+			"       -reverseTable [0|1]: Reverses the order of summary tables (set 0 to force off)\n"+
+			"       -scrollableTable [0|1]: makes the summary table scrollable, with frozen first rows and columns (set 0 to force off)\n" +
+			"       -autoColorizeTable [0|1]: enables summary table autocolorization. (set 0 to force off) \n" +
 			"       -maxSummaryTableStringLength <n>: strings longer than this will get truncated\n" +
 			"       -allowDuplicateCSVs : doesn't remove duplicate CSVs (Note: can cause summary table cache file locking issues)\n"+
 			"       -requireMetadata : ignores CSVs without metadata\n" +
@@ -650,14 +651,36 @@ namespace PerfReportTool
 			{
 				weightByColumnName = null;
 			}
-			bool reverseSort = tableInfo.bReverseSortRows || GetBoolArg("reverseTable");
-			bool bScrollableTable = tableInfo.bScrollableFormatting || GetBoolArg("scrollableTable");
+
+			// Check params and any commandline overrides
+			bool bReverseTable = tableInfo.bReverseSortRows;
+			bool? bReverseTableOption = GetOptionalBoolArg("reverseTable");
+			if (bReverseTableOption != null)
+			{
+				bReverseTable = (bool)bReverseTableOption;
+			}
+
+			bool bScrollableTable = tableInfo.bScrollableFormatting;
+			bool? bScrollableTableOption = GetOptionalBoolArg("scrollableTable");
+			if (bScrollableTableOption != null)
+			{
+				bScrollableTable = (bool)bScrollableTableOption;
+			}
+
+			bool bAutoColorizeTable = tableInfo.bAutoColorize;
+			bool? bAutoColorizeTableOption = GetOptionalBoolArg("autoColorizeTable");
+			if (bAutoColorizeTableOption != null)
+			{
+				bAutoColorizeTable = (bool)bAutoColorizeTableOption;
+			}
+
 			bool addMinMaxColumns = !GetBoolArg("noSummaryMinMax");
+
 			if (!string.IsNullOrEmpty(outputDir))
             {
                 filenameWithoutExtension = Path.Combine(outputDir, filenameWithoutExtension);
             }
-            SummaryTable filteredTable = table.SortAndFilter(tableInfo.columnFilterList, tableInfo.rowSortList, reverseSort, weightByColumnName);
+            SummaryTable filteredTable = table.SortAndFilter(tableInfo.columnFilterList, tableInfo.rowSortList, bReverseTable, weightByColumnName);
 			if (bCollated)
 			{
 				filteredTable = filteredTable.CollateSortedTable(tableInfo.rowSortList, addMinMaxColumns);
@@ -671,7 +694,7 @@ namespace PerfReportTool
 				filteredTable.ApplyDisplayNameMapping(statDisplaynameMapping);
 				string VersionString = GetBoolArg("noWatermarks") ? "" : Version.Get();
 				string summaryTitle = GetArg("summaryTitle", null);
-				filteredTable.WriteToHTML(filenameWithoutExtension+".html", VersionString, bSpreadsheetFriendlyStrings, tableInfo.sectionBoundaries, bScrollableTable, addMinMaxColumns, GetIntArg("maxSummaryTableStringLength", -1), reportXML.columnFormatInfoList, weightByColumnName, summaryTitle);
+				filteredTable.WriteToHTML(filenameWithoutExtension+".html", VersionString, bSpreadsheetFriendlyStrings, tableInfo.sectionBoundaries, bScrollableTable, bAutoColorizeTable, addMinMaxColumns, GetIntArg("maxSummaryTableStringLength", -1), reportXML.columnFormatInfoList, weightByColumnName, summaryTitle);
 			}
 		}
 
diff --git a/Engine/Source/Programs/CSVTools/PerfReportTool/SummaryTable.cs b/Engine/Source/Programs/CSVTools/PerfReportTool/SummaryTable.cs
--- a/Engine/Source/Programs/CSVTools/PerfReportTool/SummaryTable.cs
+++ b/Engine/Source/Programs/CSVTools/PerfReportTool/SummaryTable.cs
@@ -85,6 +85,7 @@ namespace PerfSummaries
 
 			bReverseSortRows = tableElement.GetSafeAttibute<bool>("reverseSortRows", false);
 			bScrollableFormatting = tableElement.GetSafeAttibute<bool>("scrollableFormatting", false);
+			bAutoColorize = tableElement.GetSafeAttibute<bool>("autoColorize", false);
 
 			foreach (XElement sectionBoundaryEl in tableElement.Elements("sectionBoundary"))
 			{
@@ -118,6 +119,7 @@ namespace PerfSummaries
 		public List<SummarySectionBoundaryInfo> sectionBoundaries = new List<SummarySectionBoundaryInfo>();
 		public bool bReverseSortRows;
 		public bool bScrollableFormatting;
+		public bool bAutoColorize;
 		public string weightByColumn = null;
 	}
 
@@ -826,7 +828,18 @@ namespace PerfSummaries
 		}
 
 
-		public void WriteToHTML(string htmlFilename, string VersionString, bool bSpreadsheetFriendlyStrings, List<SummarySectionBoundaryInfo> sectionBoundaries, bool bScrollableTable, bool bAddMinMaxColumns, int maxColumnStringLength, SummaryTableColumnFormatInfoCollection columnFormatInfoList, string weightByColumnName, string title)
+		public void WriteToHTML(
+			string htmlFilename, 
+			string VersionString, 
+			bool bSpreadsheetFriendlyStrings, 
+			List<SummarySectionBoundaryInfo> sectionBoundaries, 
+			bool bScrollableTable, 
+			bool bAutoColorizeTable,
+			bool bAddMinMaxColumns, 
+			int maxColumnStringLength, 
+			SummaryTableColumnFormatInfoCollection columnFormatInfoList, 
+			string weightByColumnName, 
+			string title)
 		{
 			System.IO.StreamWriter htmlFile = new System.IO.StreamWriter(htmlFilename, false);
 			int statColSpan = hasMinMaxColumns ? 3 : 1;
@@ -864,6 +877,25 @@ namespace PerfSummaries
 				}
 			}
 
+			// Get format info for the columns
+			Dictionary<SummaryTableColumn, SummaryTableColumnFormatInfo> columnFormatInfoLookup = new Dictionary<SummaryTableColumn, SummaryTableColumnFormatInfo>();
+			foreach (SummaryTableColumn column in columns)
+			{
+				columnFormatInfoLookup[column] = (columnFormatInfoList != null) ? columnFormatInfoList.GetFormatInfo(column.name) : SummaryTableColumnFormatInfoCollection.DefaultColumnInfo;
+			}
+
+			// Automatically colourize the table if requested
+			if (bAutoColorizeTable)
+			{
+				foreach (SummaryTableColumn column in columns)
+				{
+					if (column.isNumeric)
+					{
+						column.ComputeAutomaticColourThresholds(columnFormatInfoLookup[column].autoColorizeMode);
+					}
+				}
+			}
+
 			if (bScrollableTable)
 			{
 				// Insert some javascript to make the columns sticky. It's not possible to do this for multiple columns with pure CSS, since you need to compute the X offset dynamically
@@ -993,25 +1025,6 @@ namespace PerfSummaries
 			htmlFile.WriteLine("</head><body>");
 			htmlFile.WriteLine("<table id='mainTable'>");
 
-			// Get format info for the columns
-			Dictionary<SummaryTableColumn, SummaryTableColumnFormatInfo> columnFormatInfoLookup = new Dictionary<SummaryTableColumn, SummaryTableColumnFormatInfo>();
-			foreach (SummaryTableColumn column in columns)
-			{
-				columnFormatInfoLookup[column] = (columnFormatInfoList != null) ? columnFormatInfoList.GetFormatInfo(column.name) : SummaryTableColumnFormatInfoCollection.DefaultColumnInfo;
-			}
-
-			// Automatically colourize the table
-			if (bScrollableTable)
-			{
-				foreach (SummaryTableColumn column in columns)
-				{
-					if (column.isNumeric)
-					{
-						column.ComputeAutomaticColourThresholds(columnFormatInfoLookup[column].autoColorizeMode);
-					}
-				}
-			}
-
 			string HeaderRow = "";
 			if (isCollated)
 			{
