PerfReportTool 4.79 - Add support for writing summary row data to JSON, either for all data or just for CSV Metadata.

Uses the following args:
  -summaryTableToJson <filename> : json filename to write summary table row data to
  -summaryTableToJsonFastMode : exit after serializing json data (skips making summary tables)
  -summaryTableToJsonMetadataOnly : only write CsvMetadata elements to json

This makes it easy to extract raw perf data and CSV metadata from a collection of PRC files. Non PRC paths are also supported.

diff --git a/Engine/Source/Programs/CSVTools/PerfReportTool/App.config b/Engine/Source/Programs/CSVTools/PerfReportTool/App.config
--- a/Engine/Source/Programs/CSVTools/PerfReportTool/App.config
+++ b/Engine/Source/Programs/CSVTools/PerfReportTool/App.config
@@ -1,9 +1,15 @@
-<?xml version="1.0" encoding="utf-8"?>
+﻿<?xml version="1.0" encoding="utf-8"?>
 <configuration>
     <startup> 
-        <supportedRuntime version="v4.0" sku=".NETFramework,Version=v4.8"/>
+        <supportedRuntime version="v4.0" sku=".NETFramework,Version=v4.8" />
     </startup>
   <runtime>
-    <gcConcurrent enabled="false"/>
+    <gcConcurrent enabled="false" />
+    <assemblyBinding xmlns="urn:schemas-microsoft-com:asm.v1">
+      <dependentAssembly>
+        <assemblyIdentity name="System.Text.Encodings.Web" publicKeyToken="cc7b13ffcd2ddd51" culture="neutral" />
+        <bindingRedirect oldVersion="0.0.0.0-5.0.0.1" newVersion="5.0.0.1" />
+      </dependentAssembly>
+    </assemblyBinding>
   </runtime>
 </configuration>
diff --git a/Engine/Source/Programs/CSVTools/PerfReportTool/PerfReportTool.cs b/Engine/Source/Programs/CSVTools/PerfReportTool/PerfReportTool.cs
--- a/Engine/Source/Programs/CSVTools/PerfReportTool/PerfReportTool.cs
+++ b/Engine/Source/Programs/CSVTools/PerfReportTool/PerfReportTool.cs
@@ -22,7 +22,7 @@ namespace PerfReportTool
 {
     class Version
     {
-        private static string VersionString = "4.78";
+        private static string VersionString = "4.79";
 
         public static string Get() { return VersionString; }
     };
@@ -156,6 +156,11 @@ namespace PerfReportTool
 			"       -minFrameCount <n> : ignore CSVs without at least this number of valid frames\n" +
 			"       -maxFileAgeDays <n> : max file age in days. CSV or PRC files older than this will be ignored\n" +
 			"\n" +
+			"Json serialization:\n" +
+			"       -summaryTableToJson <filename> : json filename to write summary table row data to\n" +
+			"       -summaryTableToJsonFastMode : exit after serializing json data (skips making summary tables)\n" +
+			"       -summaryTableToJsonMetadataOnly : only write CsvMetadata elements to json\n" +
+			"\n" +
 			"Performance args for bulk mode:\n" +
 			"       -precacheCount <n> : number of CSV files to precache in the lookahead cache (0 for no precache)\n" +
 			"       -precacheThreads <n> : number of threads to use for the CSV lookahead cache (default 8)\n" +
@@ -322,6 +327,15 @@ namespace PerfReportTool
 			}
 			statDisplaynameMapping = reportXML.GetDisplayNameMapping();
 
+			// If we're outputting row data to json, create the dict
+			SummaryTableDataJsonHelper summaryTableJsonHelper = null;
+			string summaryJsonOutFilename = GetArg("summaryTableToJson", null);
+			if (summaryJsonOutFilename != null)
+			{
+				summaryTableJsonHelper = new SummaryTableDataJsonHelper(summaryJsonOutFilename, GetBoolArg("summaryTableToJsonMetadataOnly"));
+			}
+
+
 			SummaryTableCacheStats summaryTableCacheStats = new SummaryTableCacheStats();
 
 			perfLog.LogTiming("Initialization");
@@ -552,6 +566,10 @@ namespace PerfReportTool
 							if (bIncludeRowData)
 							{
 								summaryTable.AddRowData(rowData, bReadAllStats, bShowHiddenStats);
+								if (summaryTableJsonHelper != null)
+								{
+									summaryTableJsonHelper.AddRowData(rowData);
+								}
 							}
 							perfLog.LogTiming("  AddRowData");
 						}
@@ -571,6 +589,18 @@ namespace PerfReportTool
 				}
 			}
 
+
+			if (summaryTableJsonHelper != null)
+			{
+				summaryTableJsonHelper.WriteJsonFile();
+				perfLog.LogTiming("WriteSummaryDataJson");
+				if (GetBoolArg("summaryTableToJsonFastMode"))
+				{
+					perfLog.LogTotalTiming();
+					return;
+				}
+			}
+
             Console.WriteLine("-------------------------------------------------");
 
             // Write out the summary table, if there is one
diff --git a/Engine/Source/Programs/CSVTools/PerfReportTool/PerfreportTool.csproj b/Engine/Source/Programs/CSVTools/PerfReportTool/PerfreportTool.csproj
--- a/Engine/Source/Programs/CSVTools/PerfReportTool/PerfreportTool.csproj
+++ b/Engine/Source/Programs/CSVTools/PerfReportTool/PerfreportTool.csproj
@@ -12,6 +12,8 @@
     <TargetFrameworkVersion>v4.8</TargetFrameworkVersion>
     <FileAlignment>512</FileAlignment>
     <TargetFrameworkProfile />
+    <NuGetPackageImportStamp>
+    </NuGetPackageImportStamp>
   </PropertyGroup>
   <PropertyGroup Condition=" '$(Configuration)|$(Platform)' == 'Debug|AnyCPU' ">
     <PlatformTarget>x64</PlatformTarget>
@@ -37,6 +39,15 @@
     <Reference Include="System" />
     <Reference Include="System.Core" />
     <Reference Include="System.Drawing" />
+    <Reference Include="System.Runtime.CompilerServices.Unsafe, Version=4.0.6.0, Culture=neutral, PublicKeyToken=b03f5f7f11d50a3a, processorArchitecture=MSIL">
+      <HintPath>..\packages\System.Runtime.CompilerServices.Unsafe.4.7.1\lib\net461\System.Runtime.CompilerServices.Unsafe.dll</HintPath>
+    </Reference>
+    <Reference Include="System.Text.Encodings.Web, Version=5.0.0.1, Culture=neutral, PublicKeyToken=cc7b13ffcd2ddd51, processorArchitecture=MSIL">
+      <HintPath>..\packages\System.Text.Encodings.Web.5.0.1\lib\net461\System.Text.Encodings.Web.dll</HintPath>
+    </Reference>
+    <Reference Include="System.Text.Json, Version=4.0.1.2, Culture=neutral, PublicKeyToken=cc7b13ffcd2ddd51, processorArchitecture=MSIL">
+      <HintPath>..\packages\System.Text.Json.4.7.2\lib\net461\System.Text.Json.dll</HintPath>
+    </Reference>
     <Reference Include="System.Xml.Linq" />
     <Reference Include="System.Data.DataSetExtensions" />
     <Reference Include="Microsoft.CSharp" />
@@ -66,6 +77,7 @@
   </ItemGroup>
   <ItemGroup>
     <None Include="App.config" />
+    <None Include="packages.config" />
   </ItemGroup>
   <ItemGroup>
     <ProjectReference Include="..\CsvStats\CSVStats.csproj">
diff --git a/Engine/Source/Programs/CSVTools/PerfReportTool/SummaryTableData.cs b/Engine/Source/Programs/CSVTools/PerfReportTool/SummaryTableData.cs
--- a/Engine/Source/Programs/CSVTools/PerfReportTool/SummaryTableData.cs
+++ b/Engine/Source/Programs/CSVTools/PerfReportTool/SummaryTableData.cs
@@ -6,9 +6,46 @@ using System.Linq;
 using System.Xml.Linq;
 using PerfReportTool;
 using System.IO;
+using System.Text.Json;
 
 namespace PerfSummaries
 {
+
+	class SummaryTableDataJsonHelper
+	{
+		public SummaryTableDataJsonHelper(string InJsonFilename, bool bInCsvMetadataOnly)
+		{
+			JsonFilename = InJsonFilename;
+			bCsvMetadataOnly = bInCsvMetadataOnly;
+		}
+		public void AddRowData(SummaryTableRowData rowData)
+		{
+			if (!rowData.dict.ContainsKey("csvid"))
+			{
+				Console.WriteLine("Warning: SummaryTableDataJsonHelper.AddRowData - Row data contains no CSV ID! Skipping");
+				return;
+			}
+
+			Dict.Add(rowData.dict["csvid"].value, rowData.ToJsonDict(bCsvMetadataOnly));
+		}
+
+		public void WriteJsonFile()
+		{
+			Console.WriteLine("Writing summary table row data to json: " + JsonFilename);
+			JsonSerializerOptions options = new JsonSerializerOptions { WriteIndented = true };
+			string jsonString = JsonSerializer.Serialize(Dict, options);
+
+			// Write the file to Json and rename it when we're done
+			File.WriteAllText(JsonFilename, jsonString);
+		}
+
+		Dictionary<string, dynamic> Dict = new Dictionary<string, dynamic>();
+		string JsonFilename;
+		bool bCsvMetadataOnly;
+	}
+
+
+
 	class SummaryTableElement
 	{
 		// Bump this when making changes!
@@ -134,6 +171,45 @@ namespace PerfSummaries
 			return (flags & (uint)flag) != 0;
 		}
 
+		public Dictionary<string, dynamic> ToJsonDict(bool bWriteType)
+		{
+			Dictionary<string, dynamic> Dict = new Dictionary<string, dynamic>();
+			if (bWriteType)
+			{
+				Dict.Add("type", type.ToString());
+			}
+			if (isNumeric)
+			{
+				Dict.Add("value", numericValue);
+			}
+			else
+			{
+				Dict.Add("value", value);
+			}
+			if (!string.IsNullOrEmpty(tooltip))
+			{
+				Dict.Add("tooltip", tooltip);
+			}
+			if ( colorThresholdList != null)
+			{
+				Dict.Add("colorThresholdList", colorThresholdList.ToJsonDict());
+			}
+			List<string> FlagStrings = new List<string>();
+			var FlagValues = Enum.GetValues(typeof(Flags));
+			foreach (var FlagValue in FlagValues)
+			{
+				if ( (flags & (uint)(int)FlagValue) != 0)
+				{
+					FlagStrings.Add(FlagValue.ToString());
+				}
+			}
+			if (FlagStrings.Count > 0)
+			{
+				Dict.Add("flags", FlagStrings);
+			}
+			return Dict;
+		}
+
 		public Type type;
 		public string name;
 		public string value;
@@ -311,6 +387,37 @@ namespace PerfSummaries
 			dict.Add(key, metadataValue);
 		}
 
+		public Dictionary<string, dynamic> ToJsonDict(bool bCsvDataOnly)
+		{
+			Dictionary<string, dynamic> DictOut = new Dictionary<string, dynamic>();
+
+			// Make a dictionary for each data type
+			if (bCsvDataOnly)
+			{
+				DictOut[SummaryTableElement.Type.CsvMetadata.ToString()] = new Dictionary<string, dynamic>();
+			}
+			else
+			{
+				var DataTypes = Enum.GetValues(typeof(SummaryTableElement.Type));
+				foreach (SummaryTableElement.Type dataType in DataTypes)
+				{
+					DictOut[dataType.ToString()] = new Dictionary<string, dynamic>();
+				}
+			}
+
+			foreach (string key in dict.Keys)
+			{
+				SummaryTableElement Element = dict[key];
+				if (bCsvDataOnly && Element.type != SummaryTableElement.Type.CsvMetadata)
+				{
+					continue;
+				}
+				DictOut[Element.type.ToString()][key] = Element.ToJsonDict(false);
+			}
+
+			return DictOut;
+		}
+
 		public Dictionary<string, SummaryTableElement> dict = new Dictionary<string, SummaryTableElement>();
 	};
 
diff --git a/Engine/Source/Programs/CSVTools/PerfReportTool/SummaryTypes.cs b/Engine/Source/Programs/CSVTools/PerfReportTool/SummaryTypes.cs
--- a/Engine/Source/Programs/CSVTools/PerfReportTool/SummaryTypes.cs
+++ b/Engine/Source/Programs/CSVTools/PerfReportTool/SummaryTypes.cs
@@ -188,6 +188,36 @@ namespace PerfSummaries
 			Add(new ThresholdInfo(orangeValue, null));
 			Add(new ThresholdInfo(redValue, null));
 		}
+
+		public Dictionary<string, dynamic> ToJsonDict()
+		{
+			Dictionary<string, dynamic> Dict = new Dictionary<string, dynamic>();
+			if (Thresholds.Count > 0)
+			{
+				List<double> ThresholdValues = new List<double>();
+				List<string> ThresholdColors = new List<string>();
+
+				bool bHasValidColors = false;
+				foreach (ThresholdInfo thresholdInfo in Thresholds)
+				{
+					ThresholdValues.Add(thresholdInfo.value);
+					ThresholdColors.Add(thresholdInfo.colour == null ? "null" : thresholdInfo.colour.ToString());
+
+					if (thresholdInfo.colour != null)
+					{
+						bHasValidColors = true;
+					}
+				}
+
+				Dict["thresholds"] = ThresholdValues;
+				if (bHasValidColors)
+				{
+					Dict["colors"] = ThresholdColors;
+				}
+			}
+			return Dict;
+		}
+
 		public static string GetThresholdColour(double value, double redValue, double orangeValue, double yellowValue, double greenValue,
 			Colour redOverride = null, Colour orangeOverride = null, Colour yellowOverride = null, Colour greenOverride = null)
 		{
