PerfReportTool 4.80. By default, only write SummaryTableRow values to json. Do not include secondary data such as tooltips, budget colors and flags. To enable this data, pass -summaryTableJsonWriteAllElementData

diff --git a/Engine/Source/Programs/CSVTools/PerfReportTool/PerfReportTool.cs b/Engine/Source/Programs/CSVTools/PerfReportTool/PerfReportTool.cs
--- a/Engine/Source/Programs/CSVTools/PerfReportTool/PerfReportTool.cs
+++ b/Engine/Source/Programs/CSVTools/PerfReportTool/PerfReportTool.cs
@@ -22,7 +22,7 @@ namespace PerfReportTool
 {
     class Version
     {
-        private static string VersionString = "4.79";
+        private static string VersionString = "4.80";
 
         public static string Get() { return VersionString; }
     };
@@ -159,6 +159,7 @@ namespace PerfReportTool
 			"Json serialization:\n" +
 			"       -summaryTableToJson <filename> : json filename to write summary table row data to\n" +
 			"       -summaryTableToJsonFastMode : exit after serializing json data (skips making summary tables)\n" +
+			"       -summaryTableToJsonWriteAllElementData : write all element data, including tooltips, flags\n" +
 			"       -summaryTableToJsonMetadataOnly : only write CsvMetadata elements to json\n" +
 			"\n" +
 			"Performance args for bulk mode:\n" +
@@ -332,7 +333,7 @@ namespace PerfReportTool
 			string summaryJsonOutFilename = GetArg("summaryTableToJson", null);
 			if (summaryJsonOutFilename != null)
 			{
-				summaryTableJsonHelper = new SummaryTableDataJsonHelper(summaryJsonOutFilename, GetBoolArg("summaryTableToJsonMetadataOnly"));
+				summaryTableJsonHelper = new SummaryTableDataJsonHelper(summaryJsonOutFilename, GetBoolArg("summaryTableToJsonMetadataOnly"), GetBoolArg("summaryTableToJsonWriteAllElementData"));
 			}
 
 
diff --git a/Engine/Source/Programs/CSVTools/PerfReportTool/SummaryTableData.cs b/Engine/Source/Programs/CSVTools/PerfReportTool/SummaryTableData.cs
--- a/Engine/Source/Programs/CSVTools/PerfReportTool/SummaryTableData.cs
+++ b/Engine/Source/Programs/CSVTools/PerfReportTool/SummaryTableData.cs
@@ -13,10 +13,11 @@ namespace PerfSummaries
 
 	class SummaryTableDataJsonHelper
 	{
-		public SummaryTableDataJsonHelper(string InJsonFilename, bool bInCsvMetadataOnly)
+		public SummaryTableDataJsonHelper(string InJsonFilename, bool bInCsvMetadataOnly, bool bInWriteAllElementData)
 		{
 			JsonFilename = InJsonFilename;
 			bCsvMetadataOnly = bInCsvMetadataOnly;
+			bWriteAllElementData = bInWriteAllElementData;
 		}
 		public void AddRowData(SummaryTableRowData rowData)
 		{
@@ -26,7 +27,7 @@ namespace PerfSummaries
 				return;
 			}
 
-			Dict.Add(rowData.dict["csvid"].value, rowData.ToJsonDict(bCsvMetadataOnly));
+			Dict.Add(rowData.dict["csvid"].value, rowData.ToJsonDict(bCsvMetadataOnly, bWriteAllElementData));
 		}
 
 		public void WriteJsonFile()
@@ -42,6 +43,7 @@ namespace PerfSummaries
 		Dictionary<string, dynamic> Dict = new Dictionary<string, dynamic>();
 		string JsonFilename;
 		bool bCsvMetadataOnly;
+		bool bWriteAllElementData;
 	}
 
 
@@ -178,14 +180,8 @@ namespace PerfSummaries
 			{
 				Dict.Add("type", type.ToString());
 			}
-			if (isNumeric)
-			{
-				Dict.Add("value", numericValue);
-			}
-			else
-			{
-				Dict.Add("value", value);
-			}
+			Dict.Add("value", DynamicValue);
+
 			if (!string.IsNullOrEmpty(tooltip))
 			{
 				Dict.Add("tooltip", tooltip);
@@ -209,6 +205,17 @@ namespace PerfSummaries
 			}
 			return Dict;
 		}
+		public dynamic DynamicValue
+		{
+			get
+			{
+				if (isNumeric)
+				{
+					return numericValue;
+				}
+				return value;
+			}
+		}
 
 		public Type type;
 		public string name;
@@ -387,12 +394,12 @@ namespace PerfSummaries
 			dict.Add(key, metadataValue);
 		}
 
-		public Dictionary<string, dynamic> ToJsonDict(bool bCsvDataOnly)
+		public Dictionary<string, dynamic> ToJsonDict(bool bCsvMetadataOnly, bool bWriteAllElementData)
 		{
 			Dictionary<string, dynamic> DictOut = new Dictionary<string, dynamic>();
 
 			// Make a dictionary for each data type
-			if (bCsvDataOnly)
+			if (bCsvMetadataOnly)
 			{
 				DictOut[SummaryTableElement.Type.CsvMetadata.ToString()] = new Dictionary<string, dynamic>();
 			}
@@ -408,11 +415,18 @@ namespace PerfSummaries
 			foreach (string key in dict.Keys)
 			{
 				SummaryTableElement Element = dict[key];
-				if (bCsvDataOnly && Element.type != SummaryTableElement.Type.CsvMetadata)
+				if (bCsvMetadataOnly && Element.type != SummaryTableElement.Type.CsvMetadata)
 				{
 					continue;
 				}
-				DictOut[Element.type.ToString()][key] = Element.ToJsonDict(false);
+				if (bWriteAllElementData)
+				{
+					DictOut[Element.type.ToString()][key] = Element.ToJsonDict(false);
+				}
+				else
+				{
+					DictOut[Element.type.ToString()][key] = Element.DynamicValue;
+				}
 			}
 
 			return DictOut;
