PerfReportTool - add missing packages.config

diff --git a/Engine/Source/Programs/CSVTools/PerfReportTool/packages.config b/Engine/Source/Programs/CSVTools/PerfReportTool/packages.config
new file mode 100644
--- /dev/null
+++ b/Engine/Source/Programs/CSVTools/PerfReportTool/packages.config
@@ -0,0 +1,12 @@
+﻿<?xml version="1.0" encoding="utf-8"?>
+<packages>
+  <package id="Microsoft.Bcl.AsyncInterfaces" version="6.0.0" targetFramework="net48" />
+  <package id="System.Buffers" version="4.5.1" targetFramework="net48" />
+  <package id="System.Memory" version="4.5.5" targetFramework="net48" />
+  <package id="System.Numerics.Vectors" version="4.5.0" targetFramework="net48" />
+  <package id="System.Runtime.CompilerServices.Unsafe" version="4.7.1" targetFramework="net48" />
+  <package id="System.Text.Encodings.Web" version="5.0.1" targetFramework="net48" />
+  <package id="System.Text.Json" version="4.7.2" targetFramework="net48" />
+  <package id="System.Threading.Tasks.Extensions" version="4.5.4" targetFramework="net48" />
+  <package id="System.ValueTuple" version="4.5.0" targetFramework="net48" />
+</packages>
\ No newline at end of file
