PerfReportTool 4.81 : Add -jsonToPrcDir mode which converts json serialized summary row data exported with -summaryTableToJson -summaryTableToJsonWriteAllElementData to PRCs

diff --git a/Engine/Source/Programs/CSVTools/PerfReportTool/PerfReportTool.cs b/Engine/Source/Programs/CSVTools/PerfReportTool/PerfReportTool.cs
--- a/Engine/Source/Programs/CSVTools/PerfReportTool/PerfReportTool.cs
+++ b/Engine/Source/Programs/CSVTools/PerfReportTool/PerfReportTool.cs
@@ -18,11 +18,14 @@ using PerfSummaries;
 using System.Globalization;
 using CSVTools;
 
+using System.Text.Json;
+using System.Text.Json.Serialization;
+
 namespace PerfReportTool
 {
     class Version
     {
-        private static string VersionString = "4.80";
+        private static string VersionString = "4.81";
 
         public static string Get() { return VersionString; }
     };
@@ -74,7 +77,7 @@ namespace PerfReportTool
 	};
 
 	class Program : CommandLineTool
-    {
+	{
 		static string formatString =
 			"PerfReportTool v" + Version.Get() + "\n" +
 			"\n" +
@@ -85,27 +88,27 @@ namespace PerfReportTool
 			"\n" +
 			"Optional Args:\n" +
 			"       -reportType <e.g. flythrough, playthrough, playthroughmemory>\n" +
-			"       -reportTypeCompatCheck : do a compatibility if when specifying a report type (rather than forcing)\n"+
+			"       -reportTypeCompatCheck : do a compatibility if when specifying a report type (rather than forcing)\n" +
 			"       -graphXML <xmlfilename>\n" +
 			"       -reportXML <xmlfilename>\n" +
 			"       -reportxmlbasedir <folder>\n" +
 			"       -title <name> - title for detailed reports\n" +
 			"       -summaryTitle <name> - title for summary tables\n" +
 			"       -maxy <value> - forces all graphs to use this value\n" +
-			"       -writeSummaryCsv : if specified, a csv file containing summary information will be generated.\n"+
+			"       -writeSummaryCsv : if specified, a csv file containing summary information will be generated.\n" +
 			"          Not available in bulk mode.\n" +
-			"       -noWatermarks : don't embed the commandline or version in reports\n"+
-			"       -cleanCsvOut <filename> : write a standard format CSV after event stripping with metadata stripped out.\n"+
+			"       -noWatermarks : don't embed the commandline or version in reports\n" +
+			"       -cleanCsvOut <filename> : write a standard format CSV after event stripping with metadata stripped out.\n" +
 			"          Not available in bulk mode.\n" +
 			"       -noSmooth : disable smoothing on all graphs\n" +
 			"       -listSummaryTables: lists available summary tables from the current report XML\n" +
 			"\n" +
 			"Performance args:\n" +
 			"       -perfLog : output performance logging information\n" +
-			"       -graphThreads : use with -batchedGraphs to control the number of threads per CsvToSVG instance \n"+
+			"       -graphThreads : use with -batchedGraphs to control the number of threads per CsvToSVG instance \n" +
 			"                       (default: PC core count/2)\n" +
 			"       -csvToSvgSequential : Run CsvToSvg sequentially\n" +
-			"Deprecated performance args:\n"+
+			"Deprecated performance args:\n" +
 			"       -csvToSvgProcesses : Use separate processes for csvToSVG instead of threads (slower)\n" +
 			"       -noBatchedGraphs : disable batched/multithreaded graph generation (use with -csvToSvgProcesses. Default is enabled)\n" +
 			"\n" +
@@ -128,24 +131,24 @@ namespace PerfReportTool
 			"       -csvTable : writes the summary table in CSV format instead of html\n" +
 			"       -summaryTableXML <XML filename>\n" +
 			"       -summaryTable <name> :\n" +
-			"           Selects a custom summary table type from the list in reportTypes.xml \n"+
+			"           Selects a custom summary table type from the list in reportTypes.xml \n" +
 			"           (if not specified, 'default' will be used)\n" +
 			"       -condensedSummaryTable <name> :\n" +
-			"           Selects a custom condensed summary table type from the list in reportTypes.xml \n"+
+			"           Selects a custom condensed summary table type from the list in reportTypes.xml \n" +
 			"           (if not specified, 'condensed' will be used)\n" +
-			"       -summaryTableFilename <name> : use the specified filename for the summary table (instead of SummaryTable.html)\n"+
-			"       -metadataFilter <query> or <key0=value0,key1=value1...>: filters based on CSV metadata,\n"+
+			"       -summaryTableFilename <name> : use the specified filename for the summary table (instead of SummaryTable.html)\n" +
+			"       -metadataFilter <query> or <key0=value0,key1=value1...>: filters based on CSV metadata,\n" +
 			"           e.g \"platform=ps4 AND deviceprofile=ps4_60\" \n" +
 			"       -readAllStats : allows any CSV stat avg to appear in the summary table, not just those referenced in summaries\n" +
 			"       -showHiddenStats : shows stats which have been automatically hidden (typically duplicate csv unit stats)\n" +
 			"       -externalGraphs : enables external graphs (off by default)\n" +
 			"       -spreadsheetfriendly: outputs a single quote before non-numeric entries in summary tables\n" +
 			"       -noSummaryMinMax: don't make min/max columns for each stat in a condensed summary\n" +
-			"       -reverseTable [0|1]: Reverses the order of summary tables (set 0 to force off)\n"+
+			"       -reverseTable [0|1]: Reverses the order of summary tables (set 0 to force off)\n" +
 			"       -scrollableTable [0|1]: makes the summary table scrollable, with frozen first rows and columns (set 0 to force off)\n" +
 			"       -autoColorizeTable [0|1]: enables summary table autocolorization. (set 0 to force off) \n" +
 			"       -maxSummaryTableStringLength <n>: strings longer than this will get truncated\n" +
-			"       -allowDuplicateCSVs : doesn't remove duplicate CSVs (Note: can cause summary table cache file locking issues)\n"+
+			"       -allowDuplicateCSVs : doesn't remove duplicate CSVs (Note: can cause summary table cache file locking issues)\n" +
 			"       -requireMetadata : ignores CSVs without metadata\n" +
 			"       -listFiles : just list all files that pass the metadata query. Don't generate any reports.\n" +
 			"       -reportLinkRootPath <path> : Make report links relative to this\n" +
@@ -161,50 +164,51 @@ namespace PerfReportTool
 			"       -summaryTableToJsonFastMode : exit after serializing json data (skips making summary tables)\n" +
 			"       -summaryTableToJsonWriteAllElementData : write all element data, including tooltips, flags\n" +
 			"       -summaryTableToJsonMetadataOnly : only write CsvMetadata elements to json\n" +
+			"       -jsonToPrcs <json filename> : write PRCs. PRC files will be written to -summaryTableCache folder\n" +
 			"\n" +
 			"Performance args for bulk mode:\n" +
 			"       -precacheCount <n> : number of CSV files to precache in the lookahead cache (0 for no precache)\n" +
 			"       -precacheThreads <n> : number of threads to use for the CSV lookahead cache (default 8)\n" +
-			"       -summaryTableCache <dir> : specifies a directory for summary table data to be cached.\n"+
+			"       -summaryTableCache <dir> : specifies a directory for summary table data to be cached.\n" +
 			"           This avoids processing csvs on subsequent runs when -noDetailedReports is specified\n" +
-			"           Note: Enables -readAllStats implicitly. \n"+
+			"           Note: Enables -readAllStats implicitly. \n" +
 			"       -summaryTableCacheInvalidate : regenerates summary table disk cache entries (ie write only)\n" +
 			"       -summaryTableCacheReadOnly : only read from the cache, never write\n" +
 			"       -summaryTableCachePurgeInvalid : Purges invalid PRCs from the cache folder\n" +
 			"       -summaryTableCacheIn <dir> : reads data directly from the summary table cache instead of from CSVs\n" +
-			"       -summaryTableCacheUseOnlyCsvID : only use the CSV ID for the summary table cacheID, ignoringthe report type hash\n"+
+			"       -summaryTableCacheUseOnlyCsvID : only use the CSV ID for the summary table cacheID, ignoringthe report type hash\n" +
 			"            Use this if you want to avoid cache data being invalidated by report changes\n" +
 			"       -noCsvCacheFiles: disables usage of .csv.cache files. Cache files can be much faster if filtering on metadata\n" +
 			"";
-			/*
-			"Note on custom tables:\n" +
-			"       The -customTable and -customTableSort args allow you to generate a custom summary table\n" +
-			"       This is an alternative to using preset summary tables (see -summarytable)\n" +
-			"       Example:\n"+
-			"               -customTableSort \"deviceprofile,buildversion\" -customTable \"deviceprofile,buildversion,memoryfreeMB*\" \n"+
-			"       This outputs a table containing deviceprofile, buildversion and memoryfree stats, sorted by deviceprofile and then buildversion\n" +
-			""
-			*/
+		/*
+		"Note on custom tables:\n" +
+		"       The -customTable and -customTableSort args allow you to generate a custom summary table\n" +
+		"       This is an alternative to using preset summary tables (see -summarytable)\n" +
+		"       Example:\n"+
+		"               -customTableSort \"deviceprofile,buildversion\" -customTable \"deviceprofile,buildversion,memoryfreeMB*\" \n"+
+		"       This outputs a table containing deviceprofile, buildversion and memoryfree stats, sorted by deviceprofile and then buildversion\n" +
+		""
+		*/
 
 		Dictionary<string, string> statDisplaynameMapping;
 		ReportXML reportXML;
 
 		string GetBaseDirectory()
-        {
-            string location = System.Reflection.Assembly.GetEntryAssembly().Location.ToLower();
-
-            string baseDirectory = location;
-            baseDirectory = baseDirectory.Replace("perfreporttool.exe", "");
-
-            string debugSubDir = "\\bin\\debug\\";
-            if (baseDirectory.ToLower().EndsWith(debugSubDir))
-            {
-                // Might be best to use the CSVToSVG from source, but that might not be built, so use the one checked into binaries instead
-                baseDirectory = baseDirectory.Substring(0, baseDirectory.Length - debugSubDir.Length);
-                baseDirectory += "\\..\\..\\..\\..\\Binaries\\DotNET\\CsvTools";
-            }
-            return baseDirectory;
-        }
+		{
+			string location = System.Reflection.Assembly.GetEntryAssembly().Location.ToLower();
+
+			string baseDirectory = location;
+			baseDirectory = baseDirectory.Replace("perfreporttool.exe", "");
+
+			string debugSubDir = "\\bin\\debug\\";
+			if (baseDirectory.ToLower().EndsWith(debugSubDir))
+			{
+				// Might be best to use the CSVToSVG from source, but that might not be built, so use the one checked into binaries instead
+				baseDirectory = baseDirectory.Substring(0, baseDirectory.Length - debugSubDir.Length);
+				baseDirectory += "\\..\\..\\..\\..\\Binaries\\DotNET\\CsvTools";
+			}
+			return baseDirectory;
+		}
 
 		void Run(string[] args)
 		{
@@ -220,10 +224,23 @@ namespace PerfReportTool
 			WriteLine("PerfReportTool v" + Version.Get());
 
 			ReadCommandLine(args);
-            PerfLog perfLog = new PerfLog(GetBoolArg("perfLog"));
+			PerfLog perfLog = new PerfLog(GetBoolArg("perfLog"));
 
 			SummaryFactory.Init();
 
+			// Handle converting json to PRCs if requested
+			string jsonToPrcFilename = GetArg("jsonToPrcs", null);
+			if (jsonToPrcFilename != null)
+			{
+				string prcOutputDir = GetArg("summaryTableCache", null);
+				if (prcOutputDir == null)
+				{
+					throw new Exception("-jsonToPRCs requires -summaryTableCache!");
+				}
+				ConvertJsonToPrcs(jsonToPrcFilename, prcOutputDir);
+				return;
+			}
+
 			string csvDir = null;
 
 			bool bBulkMode = false;
@@ -237,7 +254,7 @@ namespace PerfReportTool
 			}
 			else
 			{
-				csvDir=GetArg("csvDir");
+				csvDir = GetArg("csvDir");
 				int maxFileAgeDays = GetIntArg("maxFileAgeDays", -1);
 				string summaryTableCacheInDir = GetArg("summaryTableCacheIn");
 				string csvListStr = GetArg("csvList");
@@ -294,7 +311,7 @@ namespace PerfReportTool
 					bSummaryTableCacheOnlyMode = true;
 				}
 				else
-				{ 
+				{
 					string csvFilenamesStr = GetArg("csv");
 					if (csvFilenamesStr.Length == 0)
 					{
@@ -318,22 +335,22 @@ namespace PerfReportTool
 			if (GetBoolArg("listSummaryTables"))
 			{
 				Console.WriteLine("Listing summary tables:");
-				List<string> summaryTableNames=reportXML.GetSummaryTableNames();
+				List<string> summaryTableNames = reportXML.GetSummaryTableNames();
 
 				foreach (string name in summaryTableNames)
 				{
-					Console.WriteLine("  "+name);
+					Console.WriteLine("  " + name);
 				}
 				return;
 			}
 			statDisplaynameMapping = reportXML.GetDisplayNameMapping();
 
 			// If we're outputting row data to json, create the dict
-			SummaryTableDataJsonHelper summaryTableJsonHelper = null;
+			SummaryTableDataJsonWriteHelper summaryTableJsonHelper = null;
 			string summaryJsonOutFilename = GetArg("summaryTableToJson", null);
 			if (summaryJsonOutFilename != null)
 			{
-				summaryTableJsonHelper = new SummaryTableDataJsonHelper(summaryJsonOutFilename, GetBoolArg("summaryTableToJsonMetadataOnly"), GetBoolArg("summaryTableToJsonWriteAllElementData"));
+				summaryTableJsonHelper = new SummaryTableDataJsonWriteHelper(summaryJsonOutFilename, GetBoolArg("summaryTableToJsonMetadataOnly"), GetBoolArg("summaryTableToJsonWriteAllElementData"));
 			}
 
 
@@ -353,14 +370,14 @@ namespace PerfReportTool
 					List<string> badOptions = new List<string>();
 					foreach (string option in incompatibleOptionsList)
 					{
-						if ( GetArg(option, null) != null)
+						if (GetArg(option, null) != null)
 						{
 							badOptions.Add(option);
 						}
 					}
-					if (badOptions.Count>0)
+					if (badOptions.Count > 0)
 					{
-						Console.WriteLine("Warning: Summary Table cache disabled due to incompatible options ("+ string.Join(", ", badOptions) + "). See help for details.");
+						Console.WriteLine("Warning: Summary Table cache disabled due to incompatible options (" + string.Join(", ", badOptions) + "). See help for details.");
 						summaryTableCacheDir = null;
 					}
 					else
@@ -368,42 +385,42 @@ namespace PerfReportTool
 						Console.WriteLine("Using summary table cache: " + summaryTableCacheDir);
 						Directory.CreateDirectory(summaryTableCacheDir);
 
-						if ( GetBoolArg("summaryTableCachePurgeInvalid"))
+						if (GetBoolArg("summaryTableCachePurgeInvalid"))
 						{
-							Console.WriteLine("Purging invalid data from the summary table cache." );
+							Console.WriteLine("Purging invalid data from the summary table cache.");
 							DirectoryInfo di = new DirectoryInfo(summaryTableCacheDir);
 							FileInfo[] files = di.GetFiles("*.prc", SearchOption.TopDirectoryOnly);
-							int numFilesDeleted=0;
+							int numFilesDeleted = 0;
 							foreach (FileInfo file in files)
 							{
-								if ( SummaryTableRowData.TryReadFromCache(summaryTableCacheDir, file.Name.Substring(0,file.Name.Length-4))==null )
+								if (SummaryTableRowData.TryReadFromCache(summaryTableCacheDir, file.Name.Substring(0, file.Name.Length - 4)) == null)
 								{
 									File.Delete(file.FullName);
 									numFilesDeleted++;
 								}
 							}
 							summaryTableCacheStats.PurgeCount = numFilesDeleted;
-							Console.WriteLine(numFilesDeleted+" of "+files.Length+" cache entries deleted");
+							Console.WriteLine(numFilesDeleted + " of " + files.Length + " cache entries deleted");
 							perfLog.LogTiming("PurgeSummaryTableCache");
 						}
 					}
 				}
 			}
 
-            // Create the output directory if requested
-            string outputDir = GetArg("o", false).ToLower();
-            if (!string.IsNullOrEmpty(outputDir))
-            {
-                if (!Directory.Exists(outputDir))
-                {
-                    Directory.CreateDirectory(outputDir);
-                }
-            }
+			// Create the output directory if requested
+			string outputDir = GetArg("o", false).ToLower();
+			if (!string.IsNullOrEmpty(outputDir))
+			{
+				if (!Directory.Exists(outputDir))
+				{
+					Directory.CreateDirectory(outputDir);
+				}
+			}
 
 			int precacheCount = GetIntArg("precacheCount", 8);
 			int precacheThreads = GetIntArg("precacheThreads", 8);
 			bool bBatchedGraphs = true;
-			if ( GetBoolArg("noBatchedGraphs") )
+			if (GetBoolArg("noBatchedGraphs"))
 			{
 				bBatchedGraphs = false;
 			}
@@ -447,7 +464,7 @@ namespace PerfReportTool
 			}
 
 			string summaryTableCacheForRead = summaryTableCacheDir;
-			if (bSummaryTableCacheInvalidate || writeDetailedReports )
+			if (bSummaryTableCacheInvalidate || writeDetailedReports)
 			{
 				// Don't read from the summary metadata cache if we're generating full reports
 				summaryTableCacheForRead = null;
@@ -479,37 +496,37 @@ namespace PerfReportTool
 
 
 			CsvFileCache csvFileCache = new CsvFileCache(
-				csvFilenames, 
-				precacheCount, 
-				precacheThreads, 
-				!GetBoolArg("noCsvCacheFiles"), 
-				metadataQuery, 
-				reportXML, 
-				reportTypeParams, 
-				bBulkMode, 
-				bSummaryTableCacheOnlyMode, 
-				bSummaryTableCacheUseOnlyCsvID, 
+				csvFilenames,
+				precacheCount,
+				precacheThreads,
+				!GetBoolArg("noCsvCacheFiles"),
+				metadataQuery,
+				reportXML,
+				reportTypeParams,
+				bBulkMode,
+				bSummaryTableCacheOnlyMode,
+				bSummaryTableCacheUseOnlyCsvID,
 				bRemoveDuplicates,
 				bRequireMetadata,
 				summaryTableCacheForRead,
 				bListFilesMode);
 
-            SummaryTable summaryTable = new SummaryTable();
+			SummaryTable summaryTable = new SummaryTable();
 			bool bWriteToSummaryTableCache = summaryTableCacheDir != null && !bSummaryTableCacheReadonly;
 
 			int csvCount = csvFilenames.Length;
-			for ( int i=0; i<csvCount; i++)
+			for (int i = 0; i < csvCount; i++)
 			{
-                try
-                {
-                    CachedCsvFile cachedCsvFile = csvFileCache.GetNextCachedCsvFile();
+				try
+				{
+					CachedCsvFile cachedCsvFile = csvFileCache.GetNextCachedCsvFile();
 					if (cachedCsvFile == null)
 					{
 						continue;
 					}
-                    Console.WriteLine("-------------------------------------------------");
-                    Console.WriteLine("CSV " + (i+1) + "/" + csvFilenames.Length ) ;
-                    Console.WriteLine(cachedCsvFile.filename);
+					Console.WriteLine("-------------------------------------------------");
+					Console.WriteLine("CSV " + (i + 1) + "/" + csvFilenames.Length);
+					Console.WriteLine(cachedCsvFile.filename);
 
 					perfLog.LogTiming("  CsvCacheRead");
 					if (cachedCsvFile == null)
@@ -517,7 +534,7 @@ namespace PerfReportTool
 						Console.WriteLine("Skipped!");
 					}
 					else
-                    {
+					{
 						SummaryTableRowData rowData = cachedCsvFile.cachedSummaryTableRowData;
 						if (rowData == null)
 						{
@@ -556,7 +573,7 @@ namespace PerfReportTool
 						}
 
 						if (rowData != null)
-                        {
+						{
 							// Filter row based on framecount if minFrameCount is specified
 							bool bIncludeRowData = true;
 							if (frameCountThreshold > 0 && rowData.GetFrameCount() < frameCountThreshold)
@@ -575,12 +592,12 @@ namespace PerfReportTool
 							perfLog.LogTiming("  AddRowData");
 						}
 					}
-                }
+				}
 				catch (Exception e)
 				{
 					if (bBulkMode)
 					{
-						Console.Out.WriteLine("[ERROR] : "+ e.Message);
+						Console.Out.WriteLine("[ERROR] : " + e.Message);
 					}
 					else
 					{
@@ -602,17 +619,17 @@ namespace PerfReportTool
 				}
 			}
 
-            Console.WriteLine("-------------------------------------------------");
+			Console.WriteLine("-------------------------------------------------");
 
-            // Write out the summary table, if there is one
-            if (summaryTable.Count > 0)
+			// Write out the summary table, if there is one
+			if (summaryTable.Count > 0)
 			{
 				// Pre-sort the summary table to ensure determinism
-				summaryTable = summaryTable.SortRows(new List<string>(new string[] { "csvfilename" }),true);
+				summaryTable = summaryTable.SortRows(new List<string>(new string[] { "csvfilename" }), true);
 				perfLog.LogTiming("PreSort Summary table");
 
 				string summaryTableFilename = GetArg("summaryTableFilename", "SummaryTable");
-				if ( summaryTableFilename.ToLower().EndsWith(".html"))
+				if (summaryTableFilename.ToLower().EndsWith(".html"))
 				{
 					summaryTableFilename = summaryTableFilename.Substring(0, summaryTableFilename.Length - 5);
 				}
@@ -631,7 +648,7 @@ namespace PerfReportTool
 
 					if (bCollateTable)
 					{
-						WriteSummaryTableReport(outputDir, summaryTableFilename+"_Collated", summaryTable, customSummaryTableFilter.Split(',').ToList(), customSummaryTableRowSort.Split(',').ToList(), true, bCsvTable, bSpreadsheetFriendlyStrings, null, weightByColumnName);
+						WriteSummaryTableReport(outputDir, summaryTableFilename + "_Collated", summaryTable, customSummaryTableFilter.Split(',').ToList(), customSummaryTableRowSort.Split(',').ToList(), true, bCsvTable, bSpreadsheetFriendlyStrings, null, weightByColumnName);
 					}
 				}
 				else
@@ -642,7 +659,7 @@ namespace PerfReportTool
 						summaryTableName = "default";
 					}
 					SummaryTableInfo tableInfo = reportXML.GetSummaryTable(summaryTableName);
-					WriteSummaryTableReport(outputDir, summaryTableFilename, summaryTable, tableInfo, false, bCsvTable, bSpreadsheetFriendlyStrings, null );
+					WriteSummaryTableReport(outputDir, summaryTableFilename, summaryTable, tableInfo, false, bCsvTable, bSpreadsheetFriendlyStrings, null);
 					if (bCollateTable)
 					{
 						WriteSummaryTableReport(outputDir, summaryTableFilename, summaryTable, tableInfo, true, bCsvTable, bSpreadsheetFriendlyStrings, weightByColumnName);
@@ -650,22 +667,22 @@ namespace PerfReportTool
 				}
 
 				// EmailTable is hardcoded to use the condensed type
-				string condensedSummaryTable = GetArg("condensedSummaryTable",null);
+				string condensedSummaryTable = GetArg("condensedSummaryTable", null);
 				if (GetBoolArg("emailSummary") || GetBoolArg("emailTable") || condensedSummaryTable != null)
 				{
 					SummaryTableInfo tableInfo = reportXML.GetSummaryTable(condensedSummaryTable == null ? "condensed" : condensedSummaryTable);
-					WriteSummaryTableReport(outputDir, summaryTableFilename+"_Email", summaryTable, tableInfo, true, false, bSpreadsheetFriendlyStrings, weightByColumnName);
+					WriteSummaryTableReport(outputDir, summaryTableFilename + "_Email", summaryTable, tableInfo, true, false, bSpreadsheetFriendlyStrings, weightByColumnName);
 				}
-                perfLog.LogTiming("WriteSummaryTable");
-            }
+				perfLog.LogTiming("WriteSummaryTable");
+			}
 
-			if ( summaryTableCacheDir != null )
+			if (summaryTableCacheDir != null)
 			{
 				summaryTableCacheStats.LogStats();
 			}
 			Console.WriteLine("Duplicate CSVs skipped: " + csvFileCache.duplicateCount);
-            perfLog.LogTotalTiming();
-        }
+			perfLog.LogTotalTiming();
+		}
 
 		void WriteSummaryTableReport(string outputDir, string filenameWithoutExtension, SummaryTable table, List<string> columnFilterList, List<string> rowSortList, bool bCollated, bool bToCSV, bool bSpreadsheetFriendlyStrings, List<SummarySectionBoundaryInfo> sectionBoundaries, string weightByColumnName)
 		{
@@ -708,57 +725,57 @@ namespace PerfReportTool
 			bool addMinMaxColumns = !GetBoolArg("noSummaryMinMax");
 
 			if (!string.IsNullOrEmpty(outputDir))
-            {
-                filenameWithoutExtension = Path.Combine(outputDir, filenameWithoutExtension);
-            }
-            SummaryTable filteredTable = table.SortAndFilter(tableInfo.columnFilterList, tableInfo.rowSortList, bReverseTable, weightByColumnName);
+			{
+				filenameWithoutExtension = Path.Combine(outputDir, filenameWithoutExtension);
+			}
+			SummaryTable filteredTable = table.SortAndFilter(tableInfo.columnFilterList, tableInfo.rowSortList, bReverseTable, weightByColumnName);
 			if (bCollated)
 			{
 				filteredTable = filteredTable.CollateSortedTable(tableInfo.rowSortList, addMinMaxColumns);
 			}
 			if (bToCSV)
 			{
-				filteredTable.WriteToCSV(filenameWithoutExtension+".csv");
+				filteredTable.WriteToCSV(filenameWithoutExtension + ".csv");
 			}
 			else
 			{
 				filteredTable.ApplyDisplayNameMapping(statDisplaynameMapping);
 				string VersionString = GetBoolArg("noWatermarks") ? "" : Version.Get();
 				string summaryTitle = GetArg("summaryTitle", null);
-				filteredTable.WriteToHTML(filenameWithoutExtension+".html", VersionString, bSpreadsheetFriendlyStrings, tableInfo.sectionBoundaries, bScrollableTable, bAutoColorizeTable, addMinMaxColumns, GetIntArg("maxSummaryTableStringLength", -1), reportXML.columnFormatInfoList, weightByColumnName, summaryTitle);
+				filteredTable.WriteToHTML(filenameWithoutExtension + ".html", VersionString, bSpreadsheetFriendlyStrings, tableInfo.sectionBoundaries, bScrollableTable, bAutoColorizeTable, addMinMaxColumns, GetIntArg("maxSummaryTableStringLength", -1), reportXML.columnFormatInfoList, weightByColumnName, summaryTitle);
 			}
 		}
 
-		string ReplaceFileExtension( string path, string newExtension )
-        {
+		string ReplaceFileExtension(string path, string newExtension)
+		{
 			// Special case for .bin.csv
 			if (path.ToLower().EndsWith(".csv.bin"))
 			{
-				return path.Substring(0, path.Length - 8)+ newExtension;
+				return path.Substring(0, path.Length - 8) + newExtension;
 			}
 
-            int lastDotIndex = path.LastIndexOf('.');
-            if ( path.EndsWith("\""))
-            {
-                newExtension = newExtension + "\"";
-                if (lastDotIndex == -1)
-                {
-                    lastDotIndex = path.Length - 1;
-                }
-            }
-            else if ( lastDotIndex == -1 )
-            {
-                lastDotIndex = path.Length;
-            }
+			int lastDotIndex = path.LastIndexOf('.');
+			if (path.EndsWith("\""))
+			{
+				newExtension = newExtension + "\"";
+				if (lastDotIndex == -1)
+				{
+					lastDotIndex = path.Length - 1;
+				}
+			}
+			else if (lastDotIndex == -1)
+			{
+				lastDotIndex = path.Length;
+			}
 
-            return path.Substring(0, lastDotIndex) + newExtension; 
-        }
+			return path.Substring(0, lastDotIndex) + newExtension;
+		}
 
 		static Dictionary<string, bool> UniqueHTMLFilemameLookup = new Dictionary<string, bool>();
 
 		void WriteCleanCsv(CachedCsvFile csvFile, string outCsvFilename, ReportTypeInfo reportTypeInfo)
 		{
-			if ( File.Exists(outCsvFilename) )
+			if (File.Exists(outCsvFilename))
 			{
 				throw new Exception("Clean csv file " + outCsvFilename + " already exists!");
 			}
@@ -779,7 +796,7 @@ namespace PerfReportTool
 			csvStats.WriteToCSV(outCsvFilename, false);
 		}
 
-		CsvStats ProcessCsv(CachedCsvFile csvFile, out int numFramesStripped, out CsvStats csvStatsUnstripped, int minX=0, int maxX=Int32.MaxValue, PerfLog perfLog=null, bool bStripStatsByEvents = true)
+		CsvStats ProcessCsv(CachedCsvFile csvFile, out int numFramesStripped, out CsvStats csvStatsUnstripped, int minX = 0, int maxX = Int32.MaxValue, PerfLog perfLog = null, bool bStripStatsByEvents = true)
 		{
 			numFramesStripped = 0;
 			CsvStats csvStats = ReadCsvStats(csvFile, minX, maxX);
@@ -802,21 +819,21 @@ namespace PerfReportTool
 		}
 
 		void GenerateReport(CachedCsvFile csvFile, string outputDir, bool bBulkMode, SummaryTableRowData rowData, bool bBatchedGraphs, bool writeDetailedReport, bool bReadAllStats, ReportTypeInfo reportTypeInfo, string csvDir)
-        {
-            PerfLog perfLog = new PerfLog(GetBoolArg("perfLog"));
-            string shortName = ReplaceFileExtension(MakeShortFilename(csvFile.filename), "");
-            string title = GetArg("title", false);
-            if (title.Length == 0)
-            {
-                title = shortName;
-            }
-
-            char c = title[0];
-            c = char.ToUpper(c);
-            title = c + title.Substring(1);
-
-            int minX = GetIntArg("minx", 0);
-            int maxX = GetIntArg("maxx", Int32.MaxValue);
+		{
+			PerfLog perfLog = new PerfLog(GetBoolArg("perfLog"));
+			string shortName = ReplaceFileExtension(MakeShortFilename(csvFile.filename), "");
+			string title = GetArg("title", false);
+			if (title.Length == 0)
+			{
+				title = shortName;
+			}
+
+			char c = title[0];
+			c = char.ToUpper(c);
+			title = c + title.Substring(1);
+
+			int minX = GetIntArg("minx", 0);
+			int maxX = GetIntArg("maxx", Int32.MaxValue);
 
 			string htmlFilename = null;
 			if (writeDetailedReport)
@@ -839,7 +856,7 @@ namespace PerfReportTool
 					UniqueHTMLFilemameLookup.Add(htmlFilename.Trim().ToLower(), true);
 				}
 				htmlFilename += ".html";
-            }
+			}
 
 			bool bCsvToSvgProcesses = GetBoolArg("csvToSvgProcesses");
 			bool bCsvToSvgMultiThreaded = !GetBoolArg("csvToSvgSequential");
@@ -855,8 +872,8 @@ namespace PerfReportTool
 			string responseFilename = null;
 			List<Process> csvToSvgProcesses = new List<Process>();
 			List<Task> csvToSvgTasks = new List<Task>();
-            if (writeDetailedReport)
-            {
+			if (writeDetailedReport)
+			{
 				GraphGenerator graphGenerator = null;
 				if (!bCsvToSvgProcesses)
 				{
@@ -887,15 +904,15 @@ namespace PerfReportTool
 						svgFilename = GetTempFilename(csvFile.filename) + ".svg";
 						if (graphGenerator != null)
 						{
-							GraphParams graphParams=GetCsvToSvgGraphParams(csvFile.filename, graph, thickness, minX, maxX, false, svgFilenames.Count);
+							GraphParams graphParams = GetCsvToSvgGraphParams(csvFile.filename, graph, thickness, minX, maxX, false, svgFilenames.Count);
 							if (bCsvToSvgMultiThreaded)
 							{
-								csvToSvgTasks.Add( graphGenerator.MakeGraphAsync(graphParams, svgFilename, true, false) );
+								csvToSvgTasks.Add(graphGenerator.MakeGraphAsync(graphParams, svgFilename, true, false));
 							}
 							else
 							{
 								graphGenerator.MakeGraph(graphParams, svgFilename, true, false);
-							}							
+							}
 						}
 						else
 						{
@@ -919,11 +936,11 @@ namespace PerfReportTool
 					// Save the response file
 					responseFilename = GetTempFilename(csvFile.filename) + "_response.txt";
 					System.IO.File.WriteAllLines(responseFilename, csvToSvgCommandlines);
-					Process csvToSvgProcess = LaunchCsvToSvgAsync("-batchCommands \""+responseFilename +"\" -mt " + GetIntArg("graphThreads", Environment.ProcessorCount/2).ToString() );
+					Process csvToSvgProcess = LaunchCsvToSvgAsync("-batchCommands \"" + responseFilename + "\" -mt " + GetIntArg("graphThreads", Environment.ProcessorCount / 2).ToString());
 					csvToSvgProcesses.Add(csvToSvgProcess);
 				}
 			}
-            perfLog.LogTiming("    Initial Processing");
+			perfLog.LogTiming("    Initial Processing");
 
 			// Check if we're stripping stats
 			bool bStripStatsByEvents = reportTypeInfo.bStripEvents;
@@ -946,40 +963,40 @@ namespace PerfReportTool
 			// Read the full csv while we wait for the graph processes to complete (this is only safe for CsvToSVG processes, not task threads)
 			int numFramesStripped;
 			CsvStats csvStatsUnstripped;
-			CsvStats csvStats=ProcessCsv(csvFile, out numFramesStripped, out csvStatsUnstripped, minX, maxX, perfLog, bStripStatsByEvents);
+			CsvStats csvStats = ProcessCsv(csvFile, out numFramesStripped, out csvStatsUnstripped, minX, maxX, perfLog, bStripStatsByEvents);
 
-            if ( writeDetailedReport && csvToSvgProcesses.Count > 0)
-            { 
-                // wait on the graph processes to complete
-                foreach (Process process in csvToSvgProcesses)
+			if (writeDetailedReport && csvToSvgProcesses.Count > 0)
+			{
+				// wait on the graph processes to complete
+				foreach (Process process in csvToSvgProcesses)
 				{
 					process.WaitForExit();
 				}
-                perfLog.LogTiming("    WaitForAsyncGraphs");
-            }
+				perfLog.LogTiming("    WaitForAsyncGraphs");
+			}
 
 
 			// Generate CSV metadata
 			if (rowData != null)
 			{
-                Uri currentDirUri = new Uri(Directory.GetCurrentDirectory() + "/", UriKind.Absolute);
-                if ( outputDir.Length > 0 && !outputDir.EndsWith("/"))
-                {
-                    outputDir += "/";
-                }
-                Uri optionalDirUri = new Uri(outputDir, UriKind.RelativeOrAbsolute);
+				Uri currentDirUri = new Uri(Directory.GetCurrentDirectory() + "/", UriKind.Absolute);
+				if (outputDir.Length > 0 && !outputDir.EndsWith("/"))
+				{
+					outputDir += "/";
+				}
+				Uri optionalDirUri = new Uri(outputDir, UriKind.RelativeOrAbsolute);
 
 				// Make a Csv URI that's relative to the report directory
-                Uri finalDirUri;
-                if (optionalDirUri.IsAbsoluteUri)
-                {
-                    finalDirUri = optionalDirUri;
-                }
-                else
-                {
-                    finalDirUri = new Uri(currentDirUri,outputDir);
-                }
-                Uri csvFileUri = new Uri(csvFile.filename, UriKind.Absolute);
+				Uri finalDirUri;
+				if (optionalDirUri.IsAbsoluteUri)
+				{
+					finalDirUri = optionalDirUri;
+				}
+				else
+				{
+					finalDirUri = new Uri(currentDirUri, outputDir);
+				}
+				Uri csvFileUri = new Uri(csvFile.filename, UriKind.Absolute);
 				Uri relativeCsvUri = finalDirUri.MakeRelativeUri(csvFileUri);
 				string csvPath = relativeCsvUri.ToString();
 
@@ -987,25 +1004,25 @@ namespace PerfReportTool
 				string csvId = null;
 				if (csvStats.metaData != null)
 				{
-					csvId=csvStats.metaData.GetValue("csvid", null);
+					csvId = csvStats.metaData.GetValue("csvid", null);
 				}
 
 				// re-root the CSV path if requested
 				string csvLinkRootPath = GetArg("csvLinkRootPath", null);
-				if ( csvDir != null && csvLinkRootPath != null)
+				if (csvDir != null && csvLinkRootPath != null)
 				{
 					string csvDirFinal = csvDir.Replace("\\", "/");
 					csvDirFinal += csvDirFinal.EndsWith("/") ? "" : "/";
 					Uri csvDirUri = new Uri(csvDirFinal, UriKind.Absolute);
 					Uri csvRelativeToCsvDirUri = csvDirUri.MakeRelativeUri(csvFileUri);
-					csvPath = Path.Combine(csvLinkRootPath,csvRelativeToCsvDirUri.ToString());
+					csvPath = Path.Combine(csvLinkRootPath, csvRelativeToCsvDirUri.ToString());
 					csvPath = new Uri(csvPath, UriKind.Absolute).ToString();
 				}
 
 				string csvLink = "<a href='" + csvPath + "'>" + shortName + ".csv" + "</a>";
 				if (bLinkTemplates)
 				{
-					if ( !csvPath.StartsWith("http://") && !csvPath.StartsWith("https://") )
+					if (!csvPath.StartsWith("http://") && !csvPath.StartsWith("https://"))
 					{
 						csvLink = "{LinkTemplate:Csv:" + (csvId ?? "0") + "}";
 					}
@@ -1029,7 +1046,7 @@ namespace PerfReportTool
 						{
 							reportLink = "{LinkTemplate:Report:" + (csvId ?? "0") + "}";
 						}
-					}	
+					}
 					rowData.Add(SummaryTableElement.Type.ToolMetadata, "Report", reportLink);
 				}
 				// Pass through all the metadata from the CSV
@@ -1044,7 +1061,7 @@ namespace PerfReportTool
 				if (bReadAllStats)
 				{
 					// Add every stat avg value to the metadata
-					foreach ( StatSamples stat in csvStats.Stats.Values )
+					foreach (StatSamples stat in csvStats.Stats.Values)
 					{
 						rowData.Add(SummaryTableElement.Type.CsvStatAverage, stat.Name, (double)stat.average);
 					}
@@ -1053,18 +1070,18 @@ namespace PerfReportTool
 			}
 
 			if (htmlFilename != null && !string.IsNullOrEmpty(outputDir))
-            {
-                htmlFilename = Path.Combine(outputDir, htmlFilename);
-            }
+			{
+				htmlFilename = Path.Combine(outputDir, htmlFilename);
+			}
 
-            // Write the report
-            WriteReport(htmlFilename, title, svgFilenames, reportTypeInfo, csvStats, csvStatsUnstripped, numFramesStripped, minX, maxX, bBulkMode, rowData);
-            perfLog.LogTiming("    WriteReport");
+			// Write the report
+			WriteReport(htmlFilename, title, svgFilenames, reportTypeInfo, csvStats, csvStatsUnstripped, numFramesStripped, minX, maxX, bBulkMode, rowData);
+			perfLog.LogTiming("    WriteReport");
 
-            // Delete the temp files
-            foreach (string svgFilename in svgFilenames)
-            {
-				if(svgFilename != String.Empty && File.Exists(svgFilename))
+			// Delete the temp files
+			foreach (string svgFilename in svgFilenames)
+			{
+				if (svgFilename != String.Empty && File.Exists(svgFilename))
 				{
 					File.Delete(svgFilename);
 				}
@@ -1135,13 +1152,13 @@ namespace PerfReportTool
 		CsvStats StripCsvStatsByEvents(CsvStats csvStats, out int numFramesStripped)
 		{
 			numFramesStripped = 0;
-            List<CsvEventStripInfo> eventsToStrip = reportXML.GetCsvEventsToStrip();
+			List<CsvEventStripInfo> eventsToStrip = reportXML.GetCsvEventsToStrip();
 			// We want to run the mask apply in parallel if -nodetailedreports is specified. Otherwise leave cores free for graph generation
 			bool doParallelMaskApply = GetBoolArg("noDetailedReports");
 			CsvStats strippedStats = csvStats;
 
 			if (eventsToStrip != null)
-            {
+			{
 				BitArray sampleMask = null;
 				foreach (CsvEventStripInfo eventStripInfo in eventsToStrip)
 				{
@@ -1154,22 +1171,22 @@ namespace PerfReportTool
 				}
 			}
 
-            if (numFramesStripped > 0 )
-            {
-                Console.WriteLine("CSV frames excluded : " + numFramesStripped);
-            }
-            return strippedStats;
-        }
+			if (numFramesStripped > 0)
+			{
+				Console.WriteLine("CSV frames excluded : " + numFramesStripped);
+			}
+			return strippedStats;
+		}
+
 
-      
 
 
-        void WriteReport(string htmlFilename, string title, List<string> svgFilenames, ReportTypeInfo reportTypeInfo, CsvStats csvStats, CsvStats csvStatsUnstripped, int numFramesStripped, int minX, int maxX, bool bBulkMode, SummaryTableRowData summaryRowData)
-        {
- 
-            ReportGraph[] graphs = reportTypeInfo.graphs.ToArray();
-            string titleStr = reportTypeInfo.title + " : " + title;
-            System.IO.StreamWriter htmlFile = null;
+		void WriteReport(string htmlFilename, string title, List<string> svgFilenames, ReportTypeInfo reportTypeInfo, CsvStats csvStats, CsvStats csvStatsUnstripped, int numFramesStripped, int minX, int maxX, bool bBulkMode, SummaryTableRowData summaryRowData)
+		{
+
+			ReportGraph[] graphs = reportTypeInfo.graphs.ToArray();
+			string titleStr = reportTypeInfo.title + " : " + title;
+			System.IO.StreamWriter htmlFile = null;
 
 			if (htmlFilename != null)
 			{
@@ -1177,9 +1194,9 @@ namespace PerfReportTool
 				htmlFile.WriteLine("<html>");
 				htmlFile.WriteLine("  <head>");
 				htmlFile.WriteLine("    <meta http-equiv='X-UA-Compatible' content='IE=edge'/>");
-				if ( GetBoolArg("noWatermarks"))
+				if (GetBoolArg("noWatermarks"))
 				{
-					htmlFile.WriteLine("    <![CDATA[ \nCreated with PerfReportTool" );
+					htmlFile.WriteLine("    <![CDATA[ \nCreated with PerfReportTool");
 				}
 				else
 				{
@@ -1214,23 +1231,23 @@ namespace PerfReportTool
 
 				htmlFile.WriteLine("<table style='width:800'>");
 
-				if ( reportTypeInfo.metadataToShowList != null )
-				{
-				    Dictionary<string, string> displayNameMapping = reportXML.GetDisplayNameMapping();
-    
-				    foreach (string metadataStr in reportTypeInfo.metadataToShowList)
-				    {
-					    string value = csvStats.metaData.GetValue(metadataStr, null);
-					    if (value != null)
-					    {
-						    string friendlyName = metadataStr;
-						    if (displayNameMapping.ContainsKey(metadataStr.ToLower()))
-						    {
-							    friendlyName = displayNameMapping[metadataStr];
-						    }
-						    htmlFile.WriteLine("<tr><td bgcolor='#F0F0F0'>" + friendlyName + "</td><td><b>" + value + "</b></td></tr>");
-					    }
-				    }
+				if (reportTypeInfo.metadataToShowList != null)
+				{
+					Dictionary<string, string> displayNameMapping = reportXML.GetDisplayNameMapping();
+
+					foreach (string metadataStr in reportTypeInfo.metadataToShowList)
+					{
+						string value = csvStats.metaData.GetValue(metadataStr, null);
+						if (value != null)
+						{
+							string friendlyName = metadataStr;
+							if (displayNameMapping.ContainsKey(metadataStr.ToLower()))
+							{
+								friendlyName = displayNameMapping[metadataStr];
+							}
+							htmlFile.WriteLine("<tr><td bgcolor='#F0F0F0'>" + friendlyName + "</td><td><b>" + value + "</b></td></tr>");
+						}
+					}
 				}
 				htmlFile.WriteLine("<tr><td bgcolor='#F0F0F0'>Frame count</td><td>" + csvStats.SampleCount + " (" + numFramesStripped + " excluded)</td></tr>");
 				htmlFile.WriteLine("</table>");
@@ -1238,12 +1255,12 @@ namespace PerfReportTool
 			}
 
 			if (summaryRowData != null)
-            {
-                summaryRowData.Add(SummaryTableElement.Type.ToolMetadata, "framecount", csvStats.SampleCount.ToString());
-                if (numFramesStripped > 0)
-                {
-                    summaryRowData.Add(SummaryTableElement.Type.ToolMetadata, "framecountExcluded", numFramesStripped.ToString());
-                }
+			{
+				summaryRowData.Add(SummaryTableElement.Type.ToolMetadata, "framecount", csvStats.SampleCount.ToString());
+				if (numFramesStripped > 0)
+				{
+					summaryRowData.Add(SummaryTableElement.Type.ToolMetadata, "framecountExcluded", numFramesStripped.ToString());
+				}
 			}
 
 			bool bWriteSummaryCsv = GetBoolArg("writeSummaryCsv") && !bBulkMode;
@@ -1253,21 +1270,21 @@ namespace PerfReportTool
 			if (bExtraLinksSummary)
 			{
 				bool bLinkTemplates = GetBoolArg("linkTemplates");
-				summaries.Insert(0,new ExtraLinksSummary(null, null, bLinkTemplates));
+				summaries.Insert(0, new ExtraLinksSummary(null, null, bLinkTemplates));
 			}
 
 			// If the reporttype has summary info, then write out the summary]
 			PeakSummary peakSummary = null;
-            foreach (Summary summary in summaries)
-            {
-                summary.WriteSummaryData(htmlFile, summary.useUnstrippedCsvStats ? csvStatsUnstripped : csvStats, csvStatsUnstripped, bWriteSummaryCsv, summaryRowData, htmlFilename);
-                if ( summary.GetType() == typeof(PeakSummary) )
-                {
-                    peakSummary = (PeakSummary)summary;
-                }
-            }
+			foreach (Summary summary in summaries)
+			{
+				summary.WriteSummaryData(htmlFile, summary.useUnstrippedCsvStats ? csvStatsUnstripped : csvStats, csvStatsUnstripped, bWriteSummaryCsv, summaryRowData, htmlFilename);
+				if (summary.GetType() == typeof(PeakSummary))
+				{
+					peakSummary = (PeakSummary)summary;
+				}
+			}
 
-            if (htmlFile != null)
+			if (htmlFile != null)
 			{
 				// Output the list of graphs
 				htmlFile.WriteLine("<h2>Graphs</h2>");
@@ -1276,9 +1293,9 @@ namespace PerfReportTool
 				List<string> sections = new List<string>();
 
 				//// We have to at least have the empty string in this array so that we can print the list of links.
-				if (sections.Count() == 0) 
-				{ 
-					sections.Add(""); 
+				if (sections.Count() == 0)
+				{
+					sections.Add("");
 				}
 
 				for (int index = 0; index < sections.Count; index++)
@@ -1309,7 +1326,7 @@ namespace PerfReportTool
 
 
 				// Output the Graphs
-				for(int svgFileIndex = 0; svgFileIndex < svgFilenames.Count; svgFileIndex++)
+				for (int svgFileIndex = 0; svgFileIndex < svgFilenames.Count; svgFileIndex++)
 				{
 					string svgFilename = svgFilenames[svgFileIndex];
 					if (String.IsNullOrEmpty(svgFilename))
@@ -1355,22 +1372,22 @@ namespace PerfReportTool
 					WriteEmail(htmlFilename, title, svgFilenames, reportTypeInfo, csvStats, csvStatsUnstripped, minX, maxX, bBulkMode);
 				}
 			}
-        }
+		}
 
 
-        void WriteEmail(string htmlFilename, string title, List<string> svgFilenames, ReportTypeInfo reportTypeInfo, CsvStats csvStats, CsvStats csvStatsUnstripped, int minX, int maxX, bool bBulkMode)
-        {
-			if (htmlFilename==null)
+		void WriteEmail(string htmlFilename, string title, List<string> svgFilenames, ReportTypeInfo reportTypeInfo, CsvStats csvStats, CsvStats csvStatsUnstripped, int minX, int maxX, bool bBulkMode)
+		{
+			if (htmlFilename == null)
 			{
 				return;
 			}
-            ReportGraph[] graphs = reportTypeInfo.graphs.ToArray();
-            string titleStr = reportTypeInfo.title + " : " + title;
-            System.IO.StreamWriter htmlFile;
-            htmlFile = new System.IO.StreamWriter(htmlFilename + "email");
-            htmlFile.WriteLine("<html>");
-            htmlFile.WriteLine("  <head>");
-            htmlFile.WriteLine("    <meta http-equiv='X-UA-Compatible' content='IE=edge'/>");
+			ReportGraph[] graphs = reportTypeInfo.graphs.ToArray();
+			string titleStr = reportTypeInfo.title + " : " + title;
+			System.IO.StreamWriter htmlFile;
+			htmlFile = new System.IO.StreamWriter(htmlFilename + "email");
+			htmlFile.WriteLine("<html>");
+			htmlFile.WriteLine("  <head>");
+			htmlFile.WriteLine("    <meta http-equiv='X-UA-Compatible' content='IE=edge'/>");
 			if (GetBoolArg("noWatermarks"))
 			{
 				htmlFile.WriteLine("    <![CDATA[ \nCreated with PerfReportTool");
@@ -1381,56 +1398,56 @@ namespace PerfReportTool
 				htmlFile.WriteLine(commandLine.GetCommandLine());
 			}
 			htmlFile.WriteLine("    ]]>");
-            htmlFile.WriteLine("    <title>" + titleStr + "</title>");
-            htmlFile.WriteLine("  </head>");
-            htmlFile.WriteLine("  <body><font face='verdana'>");
-            htmlFile.WriteLine("  <h1>" + titleStr + "</h1>");
+			htmlFile.WriteLine("    <title>" + titleStr + "</title>");
+			htmlFile.WriteLine("  </head>");
+			htmlFile.WriteLine("  <body><font face='verdana'>");
+			htmlFile.WriteLine("  <h1>" + titleStr + "</h1>");
 
-            // show the range
-            if (minX > 0 || maxX < Int32.MaxValue)
-            {
-                htmlFile.WriteLine("<br><br><font size='1.5'>(CSV cropped to range " + minX + "-");
-                if (maxX < Int32.MaxValue)
-                {
-                    htmlFile.WriteLine(maxX);
-                }
-                htmlFile.WriteLine(")</font>");
-            }
+			// show the range
+			if (minX > 0 || maxX < Int32.MaxValue)
+			{
+				htmlFile.WriteLine("<br><br><font size='1.5'>(CSV cropped to range " + minX + "-");
+				if (maxX < Int32.MaxValue)
+				{
+					htmlFile.WriteLine(maxX);
+				}
+				htmlFile.WriteLine(")</font>");
+			}
 
 
-            htmlFile.WriteLine("<a href=\"[Report Link Here]\">Click here for Report w/ interactive SVGs.</a>");
-            htmlFile.WriteLine("  <h2>Summary</h2>");
+			htmlFile.WriteLine("<a href=\"[Report Link Here]\">Click here for Report w/ interactive SVGs.</a>");
+			htmlFile.WriteLine("  <h2>Summary</h2>");
 
-            htmlFile.WriteLine("Overall Runtime: [Replace Me With Runtime]");
+			htmlFile.WriteLine("Overall Runtime: [Replace Me With Runtime]");
 
 			bool bWriteSummaryCsv = GetBoolArg("writeSummaryCsv") && !bBulkMode;
 
 			// If the reporttype has summary info, then write out the summary]
 			PeakSummary peakSummary = null;
-            foreach (Summary summary in reportTypeInfo.summaries)
-            {
-                summary.WriteSummaryData(htmlFile, csvStats, csvStatsUnstripped, bWriteSummaryCsv, null, htmlFilename);
-                if (summary.GetType() == typeof(PeakSummary))
-                {
-                    peakSummary = (PeakSummary)summary;
-                }
-
-            }
-
-            htmlFile.WriteLine("  </font></body>");
-            htmlFile.WriteLine("</html>");
-            htmlFile.Close();
-
-        }
-        string StripSpaces( string str )
-        {
-            return str.Replace(" ", "");
-        }
+			foreach (Summary summary in reportTypeInfo.summaries)
+			{
+				summary.WriteSummaryData(htmlFile, csvStats, csvStatsUnstripped, bWriteSummaryCsv, null, htmlFilename);
+				if (summary.GetType() == typeof(PeakSummary))
+				{
+					peakSummary = (PeakSummary)summary;
+				}
+
+			}
+
+			htmlFile.WriteLine("  </font></body>");
+			htmlFile.WriteLine("</html>");
+			htmlFile.Close();
+
+		}
+		string StripSpaces(string str)
+		{
+			return str.Replace(" ", "");
+		}
 
 		string GetTempFilename(string csvFilename)
 		{
 			string shortFileName = MakeShortFilename(csvFilename).Replace(" ", "_");
-			return Path.Combine( Path.GetTempPath(), shortFileName+"_"+Guid.NewGuid().ToString().Substring(26));
+			return Path.Combine(Path.GetTempPath(), shortFileName + "_" + Guid.NewGuid().ToString().Substring(26));
 		}
 		string GetCsvToSvgArgs(string csvFilename, string svgFilename, ReportGraph graph, double thicknessMultiplier, int minx, int maxx, bool multipleCSVs, int graphIndex, float scaleby = 1.0f)
 		{
@@ -1555,7 +1572,7 @@ namespace PerfReportTool
 			graphParams.title = graph.title;
 
 			GraphSettings graphSettings = graph.settings;
-			graphParams.statNames.AddRange( graphSettings.statString.value.Split(',') );
+			graphParams.statNames.AddRange(graphSettings.statString.value.Split(','));
 			graphParams.lineThickness = (float)(graphSettings.thickness.value * thicknessMultiplier);
 			graphParams.smooth = graphSettings.smooth.value && !GetBoolArg("nosmooth");
 			if (graphParams.smooth)
@@ -1576,7 +1593,7 @@ namespace PerfReportTool
 			}
 			graphParams.width = (int)(graphSettings.width.value * scaleby);
 			graphParams.height = (int)(graphSettings.height.value * scaleby);
-			if ( graphSettings.stacked.isSet )
+			if (graphSettings.stacked.isSet)
 			{
 				graphParams.stacked = graphSettings.stacked.value;
 				if (graphParams.stacked)
@@ -1609,7 +1626,7 @@ namespace PerfReportTool
 			}
 			if (graphSettings.hideStatPrefix.isSet && graphSettings.hideStatPrefix.value.Length > 0)
 			{
-				graphParams.hideStatPrefixes.AddRange(graphSettings.hideStatPrefix.value.Split(' ',';'));
+				graphParams.hideStatPrefixes.AddRange(graphSettings.hideStatPrefix.value.Split(' ', ';'));
 			}
 			if (multipleCSVs)
 			{
@@ -1636,13 +1653,13 @@ namespace PerfReportTool
 				{
 					for (int i = 0; i < eventsToStrip.Count; i++)
 					{
-						graphParams.highlightEventRegions.Add((eventsToStrip[i].beginName==null) ? "" : eventsToStrip[i].beginName);
+						graphParams.highlightEventRegions.Add((eventsToStrip[i].beginName == null) ? "" : eventsToStrip[i].beginName);
 						graphParams.highlightEventRegions.Add((eventsToStrip[i].endName == null) ? "{NULL}" : eventsToStrip[i].endName);
 					}
 				}
 			}
 
-			if ( graph.minFilterStatValue.isSet )
+			if (graph.minFilterStatValue.isSet)
 			{
 				graphParams.minFilterStatValue = (float)graph.minFilterStatValue.value;
 			}
@@ -1656,7 +1673,7 @@ namespace PerfReportTool
 			}
 			graphParams.uniqueId = "Graph_" + graphIndex.ToString();
 
-			if (minx>0)
+			if (minx > 0)
 			{
 				graphParams.minX = minx;
 			}
@@ -1664,7 +1681,7 @@ namespace PerfReportTool
 			{
 				graphParams.maxX = maxx;
 			}
-			if ( graphSettings.miny.isSet )
+			if (graphSettings.miny.isSet)
 			{
 				graphParams.minY = (float)graphSettings.miny.value;
 			}
@@ -1702,80 +1719,80 @@ namespace PerfReportTool
 				binary = Host == HostPlatform.Linux ? "mono" : "/Library/Frameworks/Mono.framework/Versions/Current/Commands/mono";
 				args = csvToolPath + " " + args;
 			}
-            
-            // Generate the SVGs, multithreaded
-            ProcessStartInfo startInfo = new ProcessStartInfo(binary);
-            startInfo.Arguments = args;
-            startInfo.CreateNoWindow = true;
-            startInfo.UseShellExecute = false;
-            Process process = Process.Start(startInfo);
+
+			// Generate the SVGs, multithreaded
+			ProcessStartInfo startInfo = new ProcessStartInfo(binary);
+			startInfo.Arguments = args;
+			startInfo.CreateNoWindow = true;
+			startInfo.UseShellExecute = false;
+			Process process = Process.Start(startInfo);
 			return process;
-        }
-
-        int CountCSVs( CsvStats csvStats)
-        {
-            // Count the CSVs
-            int csvCount = 0;
-            foreach (CsvEvent ev in csvStats.Events)
-            {
-                string eventName = ev.Name;
-                if (eventName.Length > 0)
-                {
-
-                    if (eventName.Contains("CSV:") && eventName.ToLower().Contains(".csv"))
-                    {
-                        csvCount++;
-                    }
-                }
-            }
-            if (csvCount == 0)
-            {
-                csvCount = 1;
-            }
-            return csvCount;
-        }
-
-        static int Main(string[] args)
-        {
-            Program program = new Program();
-            if (Debugger.IsAttached)
-            {
-                program.Run(args);
-            }
-            else
-            {
-                try
-                {
-                    program.Run(args);
-                }
-                catch (System.Exception e)
-                {
-                    Console.WriteLine("[ERROR] " + e.Message);
-                    return 1;
-                }
-            }
-
-            return 0;
-        }
+		}
+
+		int CountCSVs(CsvStats csvStats)
+		{
+			// Count the CSVs
+			int csvCount = 0;
+			foreach (CsvEvent ev in csvStats.Events)
+			{
+				string eventName = ev.Name;
+				if (eventName.Length > 0)
+				{
+
+					if (eventName.Contains("CSV:") && eventName.ToLower().Contains(".csv"))
+					{
+						csvCount++;
+					}
+				}
+			}
+			if (csvCount == 0)
+			{
+				csvCount = 1;
+			}
+			return csvCount;
+		}
+
+		static int Main(string[] args)
+		{
+			Program program = new Program();
+			if (Debugger.IsAttached)
+			{
+				program.Run(args);
+			}
+			else
+			{
+				try
+				{
+					program.Run(args);
+				}
+				catch (System.Exception e)
+				{
+					Console.WriteLine("[ERROR] " + e.Message);
+					return 1;
+				}
+			}
+
+			return 0;
+		}
 
 		bool matchesPattern(string str, string pattern)
 		{
-			string [] patternSections=pattern.ToLower().Split('*');
+			string[] patternSections = pattern.ToLower().Split('*');
 			// Check the substrings appear in order
 			string remStr = str.ToLower();
-			for (int i = 0;i<patternSections.Length; i++)
+			for (int i = 0; i < patternSections.Length; i++)
 			{
 				int idx = remStr.IndexOf(patternSections[i]);
-				if (idx==-1)
+				if (idx == -1)
 				{
 					return false;
 				}
-				remStr = remStr.Substring(idx+patternSections[i].Length);
+				remStr = remStr.Substring(idx + patternSections[i].Length);
 			}
 			return remStr.Length == 0;
 		}
 
-		System.IO.FileInfo[] GetFilesWithSearchPattern(string directory, string searchPatternStr, bool recurse, int maxFileAgeDays=-1)
+		System.IO.FileInfo[] GetFilesWithSearchPattern(string directory, string searchPatternStr, bool recurse, int maxFileAgeDays = -1)
 		{
 			List<System.IO.FileInfo> fileList = new List<FileInfo>();
 			string[] searchPatterns = searchPatternStr.Split(';');
@@ -1785,12 +1802,12 @@ namespace PerfReportTool
 				System.IO.FileInfo[] files = di.GetFiles("*.*", recurse ? SearchOption.AllDirectories : SearchOption.TopDirectoryOnly);
 				foreach (FileInfo file in files)
 				{
-					if (maxFileAgeDays >=0)
+					if (maxFileAgeDays >= 0)
 					{
 						DateTime fileModifiedTime = file.LastWriteTimeUtc;
 						DateTime currentTime = DateTime.UtcNow;
 						TimeSpan elapsed = currentTime.Subtract(fileModifiedTime);
-						if ( elapsed.TotalHours > (double)maxFileAgeDays*24.0 )
+						if (elapsed.TotalHours > (double)maxFileAgeDays * 24.0)
 						{
 							continue;
 						}
@@ -1805,7 +1822,78 @@ namespace PerfReportTool
 			return fileList.Distinct().ToArray();
 		}
 
-    }
+		void ConvertJsonToPrcs(string jsonFilename, string prcOutputDir)
+		{
+			Console.WriteLine("Converting " + jsonFilename + " to PRCs. Output folder: " + prcOutputDir);
+			if (!Directory.Exists(prcOutputDir))
+			{
+				Directory.CreateDirectory(prcOutputDir);
+			}
+			Console.WriteLine("Reading "+jsonFilename);
+			string jsonText = File.ReadAllText(jsonFilename);
+
+			Console.WriteLine("Parsing json");
+			Dictionary<string, dynamic> jsonDict = JsonToDynamicDict(jsonText);
+
+			Console.WriteLine("Writing PRCs");
+			foreach (string csvId in jsonDict.Keys)
+			{
+				Dictionary<string, dynamic> srcDict = jsonDict[csvId];
+				SummaryTableRowData rowData = new SummaryTableRowData(srcDict);
+				rowData.WriteToCache(prcOutputDir, csvId);
+			}
+		}
+
+		Dictionary<string, dynamic> JsonToDynamicDict(string jsonStr)
+		{
+			JsonElement RootElement = JsonSerializer.Deserialize<JsonElement>((string)jsonStr, null);
+			Dictionary<string, dynamic> RootElementValue = GetJsonValue(RootElement);
+			return RootElementValue;
+		}
+
+		// .net Json support is poor, so we have to do stuff like this if we just want to read a json file to a dictionary
+		dynamic GetJsonValue(JsonElement jsonElement)
+		{
+			string jsonStr = jsonElement.GetRawText();
+			switch (jsonElement.ValueKind)
+			{
+				case JsonValueKind.Number:
+					return jsonElement.GetDouble();
+
+				case JsonValueKind.Null:
+					return null;
+
+				case JsonValueKind.True:
+					return true;
+
+				case JsonValueKind.False:
+					return false;
+
+				case JsonValueKind.String:
+					return jsonElement.GetString();
+
+				case JsonValueKind.Undefined:
+					return null;
+
+				case JsonValueKind.Array:
+					List<dynamic> ArrayValue = new List<dynamic>();
+					foreach( JsonElement element in jsonElement.EnumerateArray())
+					{
+						ArrayValue.Add(GetJsonValue(element));
+					}
+					return ArrayValue;
+
+				case JsonValueKind.Object:
+					Dictionary<string, dynamic> DictValue = new Dictionary<string, dynamic>();
+					foreach ( JsonProperty property in jsonElement.EnumerateObject() )
+					{
+						DictValue[property.Name] = GetJsonValue(property.Value);
+					}
+					return DictValue;
+			}
+			return null;
+		}
+	}
 
 	static class Extensions
 	{
diff --git a/Engine/Source/Programs/CSVTools/PerfReportTool/PerfreportTool.csproj b/Engine/Source/Programs/CSVTools/PerfReportTool/PerfreportTool.csproj
--- a/Engine/Source/Programs/CSVTools/PerfReportTool/PerfreportTool.csproj
+++ b/Engine/Source/Programs/CSVTools/PerfReportTool/PerfreportTool.csproj
@@ -39,6 +39,7 @@
     <Reference Include="System" />
     <Reference Include="System.Core" />
     <Reference Include="System.Drawing" />
+    <Reference Include="System.Memory, Version=4.0.1.1, Culture=neutral, PublicKeyToken=cc7b13ffcd2ddd51, processorArchitecture=MSIL" />
     <Reference Include="System.Runtime.CompilerServices.Unsafe, Version=4.0.6.0, Culture=neutral, PublicKeyToken=b03f5f7f11d50a3a, processorArchitecture=MSIL">
       <HintPath>..\packages\System.Runtime.CompilerServices.Unsafe.4.7.1\lib\net461\System.Runtime.CompilerServices.Unsafe.dll</HintPath>
     </Reference>
diff --git a/Engine/Source/Programs/CSVTools/PerfReportTool/SummaryTableData.cs b/Engine/Source/Programs/CSVTools/PerfReportTool/SummaryTableData.cs
--- a/Engine/Source/Programs/CSVTools/PerfReportTool/SummaryTableData.cs
+++ b/Engine/Source/Programs/CSVTools/PerfReportTool/SummaryTableData.cs
@@ -7,13 +7,14 @@ using System.Xml.Linq;
 using PerfReportTool;
 using System.IO;
 using System.Text.Json;
+using System.Threading.Tasks;
 
 namespace PerfSummaries
 {
 
-	class SummaryTableDataJsonHelper
+	class SummaryTableDataJsonWriteHelper
 	{
-		public SummaryTableDataJsonHelper(string InJsonFilename, bool bInCsvMetadataOnly, bool bInWriteAllElementData)
+		public SummaryTableDataJsonWriteHelper(string InJsonFilename, bool bInCsvMetadataOnly, bool bInWriteAllElementData)
 		{
 			JsonFilename = InJsonFilename;
 			bCsvMetadataOnly = bInCsvMetadataOnly;
@@ -31,13 +32,21 @@ namespace PerfSummaries
 		}
 
 		public void WriteJsonFile()
+		{
+			Task task = WriteJsonFileAsyncPrivate();
+			task.Wait();
+
+		}
+
+		private async Task WriteJsonFileAsyncPrivate() 
 		{
 			Console.WriteLine("Writing summary table row data to json: " + JsonFilename);
 			JsonSerializerOptions options = new JsonSerializerOptions { WriteIndented = true };
-			string jsonString = JsonSerializer.Serialize(Dict, options);
 
-			// Write the file to Json and rename it when we're done
-			File.WriteAllText(JsonFilename, jsonString);
+			// serialize JSON directly to a file
+			FileStream createStream = File.Create(JsonFilename);
+			await JsonSerializer.SerializeAsync(createStream, Dict);
+			createStream.Dispose(); // DisposeAsync doesn't exist in .Net 4.8 so we can't await it
 		}
 
 		Dictionary<string, dynamic> Dict = new Dictionary<string, dynamic>();
@@ -69,8 +78,8 @@ namespace PerfSummaries
 
 		private SummaryTableElement()
 		{
-
 		}
+
 		public SummaryTableElement(Type inType, string inName, double inValue, ColourThresholdList inColorThresholdList, string inToolTip, uint inFlags = 0)
 		{
 			type = inType;
@@ -173,6 +182,49 @@ namespace PerfSummaries
 			return (flags & (uint)flag) != 0;
 		}
 
+		public SummaryTableElement(Type typeIn, string elementName, Dictionary<string, dynamic> jsonDict)
+		{
+			name = elementName;
+			type = typeIn;
+			dynamic jsonValue = jsonDict["value"];
+
+			isNumeric = !(jsonValue is string);
+			if (isNumeric)
+			{
+				numericValue = (double)jsonValue;
+				value = numericValue.ToString();
+			}
+			else
+			{
+				value = jsonValue;
+			}
+			if (jsonDict.TryGetValue("colorThresholdList", out dynamic colorThresholdListDynamic))
+			{
+				Dictionary<string, dynamic> colorThresholdDict = colorThresholdListDynamic;
+				colorThresholdList = new ColourThresholdList(colorThresholdDict);
+			}
+
+			tooltip = "";
+			if ( jsonDict.TryGetValue("tooltip", out dynamic tooltipDynamic ) )
+			{
+				tooltip = tooltipDynamic;
+			}
+
+			if (jsonDict.TryGetValue("flags", out dynamic flagsDynamic))
+			{
+				List<dynamic> flagStrings = flagsDynamic;
+				flags = 0;
+				foreach (string flagStr in flagStrings)
+				{
+					if ( Enum.TryParse<Flags>(flagStr, out Flags flag ) )
+					{
+						flags |= (uint)flag;
+					}
+				}
+			}
+
+		}
+
 		public Dictionary<string, dynamic> ToJsonDict(bool bWriteType)
 		{
 			Dictionary<string, dynamic> Dict = new Dictionary<string, dynamic>();
@@ -180,8 +232,15 @@ namespace PerfSummaries
 			{
 				Dict.Add("type", type.ToString());
 			}
-			Dict.Add("value", DynamicValue);
-
+			if (isNumeric)
+			{
+				// Serialize as Decimal to avoid 0.1 being serialized as 0.1000000000000001 etc
+				Dict.Add("value", new Decimal(numericValue));
+			}
+			else
+			{
+				Dict.Add("value", value);
+			}
 			if (!string.IsNullOrEmpty(tooltip))
 			{
 				Dict.Add("tooltip", tooltip);
@@ -226,6 +285,7 @@ namespace PerfSummaries
 		public bool isNumeric;
 		public uint flags;
 	}
+
 	class SummaryTableRowData
 	{
 		public SummaryTableRowData()
@@ -386,6 +446,11 @@ namespace PerfSummaries
 				Console.WriteLine("Warning: Key " + key + " has already been added. Ignoring...");
 			}
 		}
+		public void Add(SummaryTableElement element)
+		{
+			dict.Add(element.name.ToLower(), element);
+		}
+
 
 		public void AddString(SummaryTableElement.Type type, string name, string value, ColourThresholdList colorThresholdList = null, string tooltip = "")
 		{
@@ -394,6 +459,23 @@ namespace PerfSummaries
 			dict.Add(key, metadataValue);
 		}
 
+		public SummaryTableRowData(Dictionary<string, dynamic> jsonDict)
+		{
+			var DataTypes = Enum.GetValues(typeof(SummaryTableElement.Type));
+			foreach (string dataTypeStr in jsonDict.Keys)
+			{
+				if ( Enum.TryParse(dataTypeStr, out SummaryTableElement.Type dataType) )
+				{
+					Dictionary<string, dynamic> DataTypeDict = jsonDict[dataTypeStr];
+					foreach(string elementName in DataTypeDict.Keys)
+					{
+						SummaryTableElement newElement = new SummaryTableElement(dataType, elementName, DataTypeDict[elementName]);
+						Add(newElement);
+					}
+				}
+			}
+		}
+
 		public Dictionary<string, dynamic> ToJsonDict(bool bCsvMetadataOnly, bool bWriteAllElementData)
 		{
 			Dictionary<string, dynamic> DictOut = new Dictionary<string, dynamic>();
@@ -421,11 +503,11 @@ namespace PerfSummaries
 				}
 				if (bWriteAllElementData)
 				{
-					DictOut[Element.type.ToString()][key] = Element.ToJsonDict(false);
+					DictOut[Element.type.ToString()][Element.name] = Element.ToJsonDict(false);
 				}
 				else
 				{
-					DictOut[Element.type.ToString()][key] = Element.DynamicValue;
+					DictOut[Element.type.ToString()][Element.name] = Element.DynamicValue;
 				}
 			}
 
diff --git a/Engine/Source/Programs/CSVTools/PerfReportTool/SummaryTypes.cs b/Engine/Source/Programs/CSVTools/PerfReportTool/SummaryTypes.cs
--- a/Engine/Source/Programs/CSVTools/PerfReportTool/SummaryTypes.cs
+++ b/Engine/Source/Programs/CSVTools/PerfReportTool/SummaryTypes.cs
@@ -152,7 +152,7 @@ namespace PerfSummaries
 			int bI = (int)(b * 255.0f);
 			return "#" + rI.ToString("x2") + gI.ToString("x2") + bI.ToString("x2");
 		}
-
+		
 		public static Colour White = new Colour(1.0f, 1.0f, 1.0f, 1.0f);
 		public static Colour Black = new Colour(0, 0, 0, 1.0f);
 		public static Colour Orange = new Colour(1.0f, 0.5f, 0.0f, 1.0f);
@@ -189,6 +189,28 @@ namespace PerfSummaries
 			Add(new ThresholdInfo(redValue, null));
 		}
 
+		public ColourThresholdList(Dictionary<string, dynamic> JsonDict)
+		{ 
+			if (JsonDict.ContainsKey("thresholds"))
+			{
+				List<dynamic> thresholdsList = JsonDict["thresholds"];
+				Thresholds = new List<ThresholdInfo>(thresholdsList.Count);
+				foreach(dynamic thresholdValue in thresholdsList)
+				{
+					Add(new ThresholdInfo(thresholdValue, null));
+				}
+				if (JsonDict.ContainsKey("colors"))
+				{
+					List<dynamic> colorsList = JsonDict["colors"];
+					int n = Math.Min(colorsList.Count, thresholdsList.Count);
+					for (int i = 0; i < n; i++)
+					{
+						Thresholds[i].colour = new Colour(colorsList[i]);
+					}
+				}
+			}
+		}
+
 		public Dictionary<string, dynamic> ToJsonDict()
 		{
 			Dictionary<string, dynamic> Dict = new Dictionary<string, dynamic>();
