PerfReportTool 4.82 - add -summaryTableToJsonFileStream param : use a file stream to write Json. Experimental but can avoid OOMs. Previously this was always on

diff --git a/Engine/Source/Programs/CSVTools/PerfReportTool/PerfReportTool.cs b/Engine/Source/Programs/CSVTools/PerfReportTool/PerfReportTool.cs
--- a/Engine/Source/Programs/CSVTools/PerfReportTool/PerfReportTool.cs
+++ b/Engine/Source/Programs/CSVTools/PerfReportTool/PerfReportTool.cs
@@ -25,7 +25,7 @@ namespace PerfReportTool
 {
     class Version
     {
-        private static string VersionString = "4.81";
+        private static string VersionString = "4.82";
 
         public static string Get() { return VersionString; }
     };
@@ -164,6 +164,7 @@ namespace PerfReportTool
 			"       -summaryTableToJsonFastMode : exit after serializing json data (skips making summary tables)\n" +
 			"       -summaryTableToJsonWriteAllElementData : write all element data, including tooltips, flags\n" +
 			"       -summaryTableToJsonMetadataOnly : only write CsvMetadata elements to json\n" +
+			"       -summaryTableToJsonFileStream : use a file stream to write Json. Experimental but can avoid OOMs\n" +
 			"       -jsonToPrcs <json filename> : write PRCs. PRC files will be written to -summaryTableCache folder\n" +
 			"\n" +
 			"Performance args for bulk mode:\n" +
@@ -610,7 +611,7 @@ namespace PerfReportTool
 
 			if (summaryTableJsonHelper != null)
 			{
-				summaryTableJsonHelper.WriteJsonFile();
+				summaryTableJsonHelper.WriteJsonFile(GetBoolArg("summaryTableToJsonFileStream"));
 				perfLog.LogTiming("WriteSummaryDataJson");
 				if (GetBoolArg("summaryTableToJsonFastMode"))
 				{
diff --git a/Engine/Source/Programs/CSVTools/PerfReportTool/SummaryTableData.cs b/Engine/Source/Programs/CSVTools/PerfReportTool/SummaryTableData.cs
--- a/Engine/Source/Programs/CSVTools/PerfReportTool/SummaryTableData.cs
+++ b/Engine/Source/Programs/CSVTools/PerfReportTool/SummaryTableData.cs
@@ -31,18 +31,26 @@ namespace PerfSummaries
 			Dict.Add(rowData.dict["csvid"].value, rowData.ToJsonDict(bCsvMetadataOnly, bWriteAllElementData));
 		}
 
-		public void WriteJsonFile()
+		public void WriteJsonFile(bool bSerializeToStream)
 		{
-			Task task = WriteJsonFileAsyncPrivate();
-			task.Wait();
+			Console.WriteLine("Writing summary table row data to json: " + JsonFilename);
+
+			JsonSerializerOptions options = new JsonSerializerOptions { WriteIndented = true };
 
+			if (bSerializeToStream)
+			{
+				Task task = WriteJsonFileAsyncPrivate(options);
+				task.Wait();
+			}
+			else
+			{
+				string jsonString = JsonSerializer.Serialize(Dict, options);
+				File.WriteAllText(JsonFilename, jsonString);
+			}
 		}
 
-		private async Task WriteJsonFileAsyncPrivate() 
+		private async Task WriteJsonFileAsyncPrivate(JsonSerializerOptions options) 
 		{
-			Console.WriteLine("Writing summary table row data to json: " + JsonFilename);
-			JsonSerializerOptions options = new JsonSerializerOptions { WriteIndented = true };
-
 			// serialize JSON directly to a file
 			FileStream createStream = File.Create(JsonFilename);
 			await JsonSerializer.SerializeAsync(createStream, Dict);
