Add option to CSVInfo to extract specified stat values from the CSV.

+ Added -showAllStats which outputs all stat samples to the json file as an array.
+ Added -statFilters which lets you specify what stats to output values for.
+ Added additional info to the help text.
+ Added option to write events with their frame numbers.

diff --git a/Engine/Source/Programs/CSVTools/CsvInfo/CSVInfo.cs b/Engine/Source/Programs/CSVTools/CsvInfo/CSVInfo.cs
--- a/Engine/Source/Programs/CSVTools/CsvInfo/CSVInfo.cs
+++ b/Engine/Source/Programs/CSVTools/CsvInfo/CSVInfo.cs
@@ -12,7 +12,7 @@ namespace CSVInfo
 {
     class Version
     {
-        private static string VersionString = "1.05";
+        private static string VersionString = "1.06";
 
         public static string Get() { return VersionString; }
     };
@@ -72,10 +72,14 @@ namespace CSVInfo
             string formatString =
                 "Format: \n" +
                 "  <csvfilename>\n"+
-				"  [-showaverages]\n"+
+				"  [-showaverages] Output stat averages.\n" +
+				"  [-showTotals] Output stat totals.\n" +
+				"  [-showAllStats] Output stat values from all frames.\n" +
+				"  [-showEvents] Output CSV Events and the frame they occured on.\n" +
 				"  [-forceFullRead] (always reads the full CSV)\n" +
 				"  [-quiet] (no logging. Just throws returns a non-zero error code if the CSV is bad)\n" +
-				"  [-toJson <filename>]";
+				"  [-toJson <filename>] Output to a json file instead of stdout.\n" +
+				"  [-statFilters <stat list>] Comma separated list of stats to output values for. Wildcards are supported (eg. LLM/*)";
 
 			// Read the command line
 			if (args.Length < 1)
@@ -91,11 +95,16 @@ namespace CSVInfo
 
             bool showAverages = GetBoolArg("showAverages");
 			bool showTotals= GetBoolArg("showTotals");
+			bool showAllStats = GetBoolArg("showAllStats");
+			bool showEvents = GetBoolArg("showEvents");
 			string jsonFilename = GetArg("toJson",false);
-			bool bReadJustHeader = !GetBoolArg("forceFullRead") && !showAverages && !showTotals;
+			string statFilterString = GetArg("statFilters",false);
+			bool bReadJustHeader = !GetBoolArg("forceFullRead") && !showAverages && !showTotals && (statFilterString.Length == 0);
+
+			string[] statFilters = statFilterString.Length > 0 ? statFilterString.Split(',') : null;
 
 			CSVStats.CsvFileInfo fileInfo = new CsvFileInfo();
-			CsvStats csvStats = CsvStats.ReadCSVFile(csvFilename, null, 0, false, fileInfo, bReadJustHeader);
+			CsvStats csvStats = CsvStats.ReadCSVFile(csvFilename, statFilters, 0, false, fileInfo, bReadJustHeader);
 
 			if ( GetBoolArg("quiet") )
 			{
@@ -128,6 +137,20 @@ namespace CSVInfo
 					}
 					jsonLines.Add("  },");
 				}
+
+				if (showEvents)
+				{
+					jsonLines.Add("  \"events\": [");
+					for (int i = 0; i < csvStats.Events.Count; ++i)
+					{
+						CsvEvent csvEvent = csvStats.Events[i];
+						string line = $"		{{\"EventName\": \"{csvEvent.Name}\", \"Frame\": {csvEvent.Frame}}}";
+						line += (i < csvStats.Events.Count - 1) ? "," : "";
+						jsonLines.Add(line);
+					}
+					jsonLines.Add("  ],");
+				}
+
 				List<string> statLines = new List<string>();
 				foreach (StatSamples stat in csvStats.Stats.Values.ToArray())
 				{
@@ -135,7 +158,7 @@ namespace CSVInfo
 				}
 				statLines.Sort();
 
-				if (showTotals || showAverages)
+				if (showTotals || showAverages || showAllStats)
 				{
 					jsonLines.Add("  \"stats\": {");
 					for (int i=0; i<statLines.Count; i++)
@@ -150,6 +173,10 @@ namespace CSVInfo
 						{
 							entries.Add("\"average\":" + csvStats.GetStat(statName).average);
 						}
+						if (showAllStats)
+						{
+							entries.Add("\"all\": [" + String.Join(",", csvStats.GetStat(statName).samples) + "]");
+						}
 						string line = "    \"" + statName + "\": {" + String.Join(",", entries) + "}";
 						if (i < statLines.Count-1)
 							line += ",";
@@ -163,7 +190,6 @@ namespace CSVInfo
 					jsonLines.Add("  \"stats\": " + ToJsonStringList(statLines));
 				}
 
-
 				jsonLines.Add("}");
 				System.IO.File.WriteAllLines(jsonFilename,jsonLines);
 				Console.Out.WriteLine("Wrote csv info to " + jsonFilename);
@@ -185,10 +211,23 @@ namespace CSVInfo
 					{
 						statLine += " (Total: " + stat.total.ToString()+") ";
 					}
+					if (showAllStats)
+					{
+						statLine += " (All: [" + String.Join(",", stat.samples) + "]) ";
+					}
 					statLines.Add(statLine);
 				}
 				statLines.Sort();
 
+				if (showEvents)
+				{
+					Console.Out.WriteLine("Events:");
+					foreach (CsvEvent csvEvent in csvStats.Events)
+					{
+						Console.Out.WriteLine($"{csvEvent.Name}: {csvEvent.Frame}");
+					}
+				}
+
 				// Write out the sorted stat names
 				Console.Out.WriteLine("Stats:");
 				foreach (string statLine in statLines)
