Updated PerfReportTool to treat extra links starting with '/' as links instead of link templates.

diff --git a/Engine/Source/Programs/CSVTools/PerfReportTool/Summaries/ExtraLinksSummary.cs b/Engine/Source/Programs/CSVTools/PerfReportTool/Summaries/ExtraLinksSummary.cs
--- a/Engine/Source/Programs/CSVTools/PerfReportTool/Summaries/ExtraLinksSummary.cs
+++ b/Engine/Source/Programs/CSVTools/PerfReportTool/Summaries/ExtraLinksSummary.cs
@@ -24,7 +24,7 @@ namespace PerfSummaries
                 ReportText = Sections[0];
                 SummaryTableText = Sections[1];
 				LinkURL = Sections[2];
-				if (!LinkURL.StartsWith("http://") && !LinkURL.StartsWith("https://"))
+				if (!LinkURL.StartsWith("http://") && !LinkURL.StartsWith("https://") && !LinkURL.StartsWith("/"))
 				{
 					LinkTemplateCsvId = inLinkTemplateCsvId;
 				}
