PerfReportTool 4.83 - add -collateTableOnly arg which only generates a collated table (skips full table generation). Also fix filename bug where -collateTable can produce a summary table that overwrites the full table

diff --git a/Engine/Source/Programs/CSVTools/PerfReportTool/PerfReportTool.cs b/Engine/Source/Programs/CSVTools/PerfReportTool/PerfReportTool.cs
--- a/Engine/Source/Programs/CSVTools/PerfReportTool/PerfReportTool.cs
+++ b/Engine/Source/Programs/CSVTools/PerfReportTool/PerfReportTool.cs
@@ -25,7 +25,7 @@ namespace PerfReportTool
 {
     class Version
     {
-        private static string VersionString = "4.82";
+        private static string VersionString = "4.83";
 
         public static string Get() { return VersionString; }
     };
@@ -127,6 +127,7 @@ namespace PerfReportTool
 			"       -customTableSort <comma separated field row sort order> (use with -customTable)\n" +
 			"       -noDetailedReports : skips individual report generation\n" +
 			"       -collateTable : writes a collated table in addition to the main one, merging by row sort\n" +
+			"       -collateTableOnly : as -collateTable, but doesn't write the standard summary table.\n" +
 			"       -emailTable : writes a condensed email-friendly table (see the 'condensed' summary table)\n" +
 			"       -csvTable : writes the summary table in CSV format instead of html\n" +
 			"       -summaryTableXML <XML filename>\n" +
@@ -636,6 +637,11 @@ namespace PerfReportTool
 				}
 				bool bCsvTable = GetBoolArg("csvTable");
 				bool bCollateTable = GetBoolArg("collateTable");
+				bool bCollateTableOnly = GetBoolArg("collateTableOnly");
+
+				bCollateTable |= bCollateTableOnly;
+				string collatedTableFilename = summaryTableFilename + ( bCollateTableOnly ? "" : "_Collated" );
+
 				bool bSpreadsheetFriendlyStrings = GetBoolArg("spreadsheetFriendly");
 				string weightByColumnName = GetArg("weightByColumn", null);
 				if (customSummaryTableFilter.Length > 0)
@@ -645,11 +651,13 @@ namespace PerfReportTool
 					{
 						customSummaryTableRowSort = "buildversion,deviceprofile";
 					}
-					WriteSummaryTableReport(outputDir, summaryTableFilename, summaryTable, customSummaryTableFilter.Split(',').ToList(), customSummaryTableRowSort.Split(',').ToList(), false, bCsvTable, bSpreadsheetFriendlyStrings, null, null);
-
+					if (!bCollateTableOnly)
+					{
+						WriteSummaryTableReport(outputDir, summaryTableFilename, summaryTable, customSummaryTableFilter.Split(',').ToList(), customSummaryTableRowSort.Split(',').ToList(), false, bCsvTable, bSpreadsheetFriendlyStrings, null, null);
+					}
 					if (bCollateTable)
 					{
-						WriteSummaryTableReport(outputDir, summaryTableFilename + "_Collated", summaryTable, customSummaryTableFilter.Split(',').ToList(), customSummaryTableRowSort.Split(',').ToList(), true, bCsvTable, bSpreadsheetFriendlyStrings, null, weightByColumnName);
+						WriteSummaryTableReport(outputDir, collatedTableFilename, summaryTable, customSummaryTableFilter.Split(',').ToList(), customSummaryTableRowSort.Split(',').ToList(), true, bCsvTable, bSpreadsheetFriendlyStrings, null, weightByColumnName);
 					}
 				}
 				else
@@ -660,10 +668,13 @@ namespace PerfReportTool
 						summaryTableName = "default";
 					}
 					SummaryTableInfo tableInfo = reportXML.GetSummaryTable(summaryTableName);
-					WriteSummaryTableReport(outputDir, summaryTableFilename, summaryTable, tableInfo, false, bCsvTable, bSpreadsheetFriendlyStrings, null);
+					if (!bCollateTableOnly)
+					{
+						WriteSummaryTableReport(outputDir, summaryTableFilename, summaryTable, tableInfo, false, bCsvTable, bSpreadsheetFriendlyStrings, null);
+					}
 					if (bCollateTable)
 					{
-						WriteSummaryTableReport(outputDir, summaryTableFilename, summaryTable, tableInfo, true, bCsvTable, bSpreadsheetFriendlyStrings, weightByColumnName);
+						WriteSummaryTableReport(outputDir, collatedTableFilename, summaryTable, tableInfo, true, bCsvTable, bSpreadsheetFriendlyStrings, weightByColumnName);
 					}
 				}
 
