PerfReportTool 4.84 - fix inconsistent summary table sorting with numeric values

diff --git a/Engine/Source/Programs/CSVTools/PerfReportTool/PerfReportTool.cs b/Engine/Source/Programs/CSVTools/PerfReportTool/PerfReportTool.cs
--- a/Engine/Source/Programs/CSVTools/PerfReportTool/PerfReportTool.cs
+++ b/Engine/Source/Programs/CSVTools/PerfReportTool/PerfReportTool.cs
@@ -25,7 +25,7 @@ namespace PerfReportTool
 {
     class Version
     {
-        private static string VersionString = "4.83";
+        private static string VersionString = "4.84";
 
         public static string Get() { return VersionString; }
     };
diff --git a/Engine/Source/Programs/CSVTools/PerfReportTool/SummaryTable.cs b/Engine/Source/Programs/CSVTools/PerfReportTool/SummaryTable.cs
--- a/Engine/Source/Programs/CSVTools/PerfReportTool/SummaryTable.cs
+++ b/Engine/Source/Programs/CSVTools/PerfReportTool/SummaryTable.cs
@@ -1229,7 +1229,7 @@ namespace PerfSummaries
 					if (columnLookup.ContainsKey(s.ToLower()))
 					{
 						SummaryTableColumn column = columnLookup[s.ToLower()];
-						key += "{" + column.GetStringValue(i) + "}";
+						key += "{" + column.GetStringValue(i,false,"0.0000000000") + "}";
 					}
 					else
 					{
