PerfReportTool 4.85 - Add hideStatPrefix summary table XML attribute to allow hiding of stat CSV prefixes in summary tables for clearer formatting. e.g "Exclusive/"

diff --git a/Engine/Source/Programs/CSVTools/PerfReportTool/PerfReportTool.cs b/Engine/Source/Programs/CSVTools/PerfReportTool/PerfReportTool.cs
--- a/Engine/Source/Programs/CSVTools/PerfReportTool/PerfReportTool.cs
+++ b/Engine/Source/Programs/CSVTools/PerfReportTool/PerfReportTool.cs
@@ -25,7 +25,7 @@ namespace PerfReportTool
 {
     class Version
     {
-        private static string VersionString = "4.84";
+        private static string VersionString = "4.85";
 
         public static string Get() { return VersionString; }
     };
@@ -754,7 +754,19 @@ namespace PerfReportTool
 				filteredTable.ApplyDisplayNameMapping(statDisplaynameMapping);
 				string VersionString = GetBoolArg("noWatermarks") ? "" : Version.Get();
 				string summaryTitle = GetArg("summaryTitle", null);
-				filteredTable.WriteToHTML(filenameWithoutExtension + ".html", VersionString, bSpreadsheetFriendlyStrings, tableInfo.sectionBoundaries, bScrollableTable, bAutoColorizeTable, addMinMaxColumns, GetIntArg("maxSummaryTableStringLength", -1), reportXML.columnFormatInfoList, weightByColumnName, summaryTitle);
+				filteredTable.WriteToHTML(
+					filenameWithoutExtension + ".html", 
+					VersionString, 
+					bSpreadsheetFriendlyStrings, 
+					tableInfo.sectionBoundaries, 
+					bScrollableTable, 
+					bAutoColorizeTable, 
+					addMinMaxColumns, 
+					tableInfo.hideStatPrefix,
+					GetIntArg("maxSummaryTableStringLength", -1), 
+					reportXML.columnFormatInfoList, 
+					weightByColumnName, 
+					summaryTitle );
 			}
 		}
 
diff --git a/Engine/Source/Programs/CSVTools/PerfReportTool/SummaryTable.cs b/Engine/Source/Programs/CSVTools/PerfReportTool/SummaryTable.cs
--- a/Engine/Source/Programs/CSVTools/PerfReportTool/SummaryTable.cs
+++ b/Engine/Source/Programs/CSVTools/PerfReportTool/SummaryTable.cs
@@ -86,6 +86,7 @@ namespace PerfSummaries
 			bReverseSortRows = tableElement.GetSafeAttibute<bool>("reverseSortRows", false);
 			bScrollableFormatting = tableElement.GetSafeAttibute<bool>("scrollableFormatting", false);
 			bAutoColorize = tableElement.GetSafeAttibute<bool>("autoColorize", false);
+			hideStatPrefix = tableElement.GetSafeAttibute<string>("hideStatPrefix");
 
 			foreach (XElement sectionBoundaryEl in tableElement.Elements("sectionBoundary"))
 			{
@@ -120,6 +121,7 @@ namespace PerfSummaries
 		public bool bReverseSortRows;
 		public bool bScrollableFormatting;
 		public bool bAutoColorize;
+		public string hideStatPrefix = null;
 		public string weightByColumn = null;
 	}
 
@@ -235,13 +237,34 @@ namespace PerfSummaries
 			return newColumn;
 		}
 
-		public string GetDisplayName()
+		public string GetDisplayName(string hideStatPrefix=null)
 		{
+			string nameOut;
 			if (displayName == null)
 			{
-				return TableUtil.FormatStatName(name);
+				// Trim the stat name suffix if necessary
+				string statName = name;
+				if (hideStatPrefix != null)
+				{
+					string prefix = "";
+					string suffix = name;
+					if (name.StartsWith("Min ") || name.StartsWith("Max ") || name.StartsWith("Avg "))
+					{
+						prefix = name.Substring(0, 4);
+						suffix = name.Substring(4);
+					}
+					if ( suffix.ToLower().StartsWith(hideStatPrefix.ToLower() ) )
+					{
+						statName = prefix + suffix.Substring(hideStatPrefix.Length);
+					}
+				}
+				nameOut = TableUtil.FormatStatName(statName);
+			}
+			else
+			{
+				nameOut = displayName;
 			}
-			return displayName;
+			return nameOut;
 		}
 
 		public void SetValue(int index, double value)
@@ -836,6 +859,7 @@ namespace PerfSummaries
 			bool bScrollableTable, 
 			bool bAutoColorizeTable,
 			bool bAddMinMaxColumns, 
+			string hideStatPrefix,
 			int maxColumnStringLength, 
 			SummaryTableColumnFormatInfoCollection columnFormatInfoList, 
 			string weightByColumnName, 
@@ -1040,7 +1064,7 @@ namespace PerfSummaries
 
 				for (int i = 0; i < firstStatColumnIndex; i++)
 				{
-					HeaderRow += "<th>" + columns[i].GetDisplayName() + "</th>";
+					HeaderRow += "<th>" + columns[i].GetDisplayName(hideStatPrefix) + "</th>";
 				}
 				if (!bAddMinMaxColumns)
 				{
@@ -1051,7 +1075,7 @@ namespace PerfSummaries
 				{
 					string prefix = "";
 					string suffix = "";
-					string statName = GetStatNameWithPrefixAndSuffix(columns[i].GetDisplayName(), out prefix, out suffix);
+					string statName = GetStatNameWithPrefixAndSuffix(columns[i].GetDisplayName(hideStatPrefix), out prefix, out suffix);
 					if ((i - 1) % statColSpan == 0)
 					{
 						TopHeaderRow += "<th colspan='" + statColSpan + "' >" + statName + suffix + "</th>";
@@ -1072,7 +1096,7 @@ namespace PerfSummaries
 			{
 				foreach (SummaryTableColumn column in columns)
 				{
-					HeaderRow += "<th>" + column.GetDisplayName() + "</th>";
+					HeaderRow += "<th>" + column.GetDisplayName(hideStatPrefix) + "</th>";
 				}
 				htmlFile.WriteLine("  <tr class='lastHeaderRow'>" + HeaderRow + "</tr>");
 			}
