PerfReportTool 4.86 - Add maxStringLength and maxStringLengthCollated XML params to  SummaryTableColumnFormatInfo so we can optionally truncate column strings.

diff --git a/Engine/Source/Programs/CSVTools/PerfReportTool/PerfReportTool.cs b/Engine/Source/Programs/CSVTools/PerfReportTool/PerfReportTool.cs
--- a/Engine/Source/Programs/CSVTools/PerfReportTool/PerfReportTool.cs
+++ b/Engine/Source/Programs/CSVTools/PerfReportTool/PerfReportTool.cs
@@ -25,7 +25,7 @@ namespace PerfReportTool
 {
     class Version
     {
-        private static string VersionString = "4.85";
+        private static string VersionString = "4.86";
 
         public static string Get() { return VersionString; }
     };
@@ -763,7 +763,7 @@ namespace PerfReportTool
 					bAutoColorizeTable, 
 					addMinMaxColumns, 
 					tableInfo.hideStatPrefix,
-					GetIntArg("maxSummaryTableStringLength", -1), 
+					GetIntArg("maxSummaryTableStringLength", Int32.MaxValue), 
 					reportXML.columnFormatInfoList, 
 					weightByColumnName, 
 					summaryTitle );
diff --git a/Engine/Source/Programs/CSVTools/PerfReportTool/SummaryTable.cs b/Engine/Source/Programs/CSVTools/PerfReportTool/SummaryTable.cs
--- a/Engine/Source/Programs/CSVTools/PerfReportTool/SummaryTable.cs
+++ b/Engine/Source/Programs/CSVTools/PerfReportTool/SummaryTable.cs
@@ -186,6 +186,8 @@ namespace PerfSummaries
 		public SummaryTableColumnFormatInfo()
 		{
 			name = "Default";
+			maxStringLength = Int32.MaxValue;
+			maxStringLengthCollated = Int32.MaxValue;
 		}
 		public SummaryTableColumnFormatInfo(XElement element)
 		{
@@ -202,10 +204,18 @@ namespace PerfSummaries
 				}
 			}
 			numericFormat = element.GetSafeAttibute<string>("numericFormat");
+			maxStringLength = element.GetSafeAttibute<int>("maxStringLength", Int32.MaxValue );
+			maxStringLengthCollated = element.GetSafeAttibute<int>("maxStringLengthCollated", Int32.MaxValue );
+			if (maxStringLengthCollated == Int32.MaxValue)
+			{
+				maxStringLengthCollated = maxStringLength;
+			}
 		}
 		public AutoColorizeMode autoColorizeMode = AutoColorizeMode.HighIsBad;
 		public string name;
 		public string numericFormat;
+		public int maxStringLength;
+		public int maxStringLengthCollated;
 	};
 
 	class SummaryTableColumn
@@ -1212,11 +1222,15 @@ namespace PerfSummaries
 					}
 					string bgColorString = (colour == null ? "" : " bgcolor=" + colour);
 					bool bold = false;
-					string numericFormat = columnFormatInfoLookup[column].numericFormat;
+
+					SummaryTableColumnFormatInfo columnFormat = columnFormatInfoLookup[column];
+					int maxStringLength = Math.Min( isCollated ? columnFormat.maxStringLengthCollated : columnFormat.maxStringLength, maxColumnStringLength);
+
+					string numericFormat = columnFormat.numericFormat;
 					string stringValue = column.GetStringValue(i, true, numericFormat);
-					if (maxColumnStringLength > 0 && stringValue.Length > maxColumnStringLength)
+					if (stringValue.Length > maxStringLength)
 					{
-						stringValue = TableUtil.SafeTruncateHtmlTableValue(stringValue, maxColumnStringLength);
+						stringValue = TableUtil.SafeTruncateHtmlTableValue(stringValue, maxStringLength);
 					}
 					if (bSpreadsheetFriendlyStrings && !column.isNumeric)
 					{
