PerfReportTool 3.87 - add a -summaryTableStatThreshold param to filter out columns based on a threshold. If all values are below the threshold, the column will be removed. Only applies to CsvStatAverage and SummaryTableMetric thresholds. This can also be specified in the SummaryTable XML

diff --git a/Engine/Source/Programs/CSVTools/PerfReportTool/PerfReportTool.cs b/Engine/Source/Programs/CSVTools/PerfReportTool/PerfReportTool.cs
--- a/Engine/Source/Programs/CSVTools/PerfReportTool/PerfReportTool.cs
+++ b/Engine/Source/Programs/CSVTools/PerfReportTool/PerfReportTool.cs
@@ -25,7 +25,7 @@ namespace PerfReportTool
 {
     class Version
     {
-        private static string VersionString = "4.86";
+        private static string VersionString = "4.87";
 
         public static string Get() { return VersionString; }
     };
@@ -159,6 +159,7 @@ namespace PerfReportTool
 			"       -noWeightedAvg : Don't use weighted averages for the collated table\n" +
 			"       -minFrameCount <n> : ignore CSVs without at least this number of valid frames\n" +
 			"       -maxFileAgeDays <n> : max file age in days. CSV or PRC files older than this will be ignored\n" +
+			"       -summaryTableStatThreshold <n> : stat/metric columns in the summarytable will be filtered out if all values < threshold\n" +
 			"\n" +
 			"Json serialization:\n" +
 			"       -summaryTableToJson <filename> : json filename to write summary table row data to\n" +
@@ -740,7 +741,9 @@ namespace PerfReportTool
 			{
 				filenameWithoutExtension = Path.Combine(outputDir, filenameWithoutExtension);
 			}
-			SummaryTable filteredTable = table.SortAndFilter(tableInfo.columnFilterList, tableInfo.rowSortList, bReverseTable, weightByColumnName);
+
+			float statThreshold = GetFloatArg("summaryTableStatThreshold", tableInfo.statThreshold);
+			SummaryTable filteredTable = table.SortAndFilter(tableInfo.columnFilterList, tableInfo.rowSortList, bReverseTable, weightByColumnName, statThreshold);
 			if (bCollated)
 			{
 				filteredTable = filteredTable.CollateSortedTable(tableInfo.rowSortList, addMinMaxColumns);
diff --git a/Engine/Source/Programs/CSVTools/PerfReportTool/SummaryTable.cs b/Engine/Source/Programs/CSVTools/PerfReportTool/SummaryTable.cs
--- a/Engine/Source/Programs/CSVTools/PerfReportTool/SummaryTable.cs
+++ b/Engine/Source/Programs/CSVTools/PerfReportTool/SummaryTable.cs
@@ -86,6 +86,7 @@ namespace PerfSummaries
 			bReverseSortRows = tableElement.GetSafeAttibute<bool>("reverseSortRows", false);
 			bScrollableFormatting = tableElement.GetSafeAttibute<bool>("scrollableFormatting", false);
 			bAutoColorize = tableElement.GetSafeAttibute<bool>("autoColorize", false);
+			statThreshold = tableElement.GetSafeAttibute<float>("statThreshold", 0.0f);
 			hideStatPrefix = tableElement.GetSafeAttibute<string>("hideStatPrefix");
 
 			foreach (XElement sectionBoundaryEl in tableElement.Elements("sectionBoundary"))
@@ -121,6 +122,7 @@ namespace PerfSummaries
 		public bool bReverseSortRows;
 		public bool bScrollableFormatting;
 		public bool bAutoColorize;
+		public float statThreshold;
 		public string hideStatPrefix = null;
 		public string weightByColumn = null;
 	}
@@ -227,19 +229,21 @@ namespace PerfSummaries
 		List<double> doubleValues = new List<double>();
 		List<string> stringValues = new List<string>();
 		List<string> toolTips = new List<string>();
+		public SummaryTableElement.Type elementType;
 
 		List<ColourThresholdList> colourThresholds = new List<ColourThresholdList>();
 		ColourThresholdList colourThresholdOverride = null;
-		public SummaryTableColumn(string inName, bool inIsNumeric, string inDisplayName = null, bool inIsRowWeightColumn = false)
+		public SummaryTableColumn(string inName, bool inIsNumeric, string inDisplayName, bool inIsRowWeightColumn, SummaryTableElement.Type inElementType)
 		{
 			name = inName;
 			isNumeric = inIsNumeric;
 			displayName = inDisplayName;
 			isRowWeightColumn = inIsRowWeightColumn;
+			elementType = inElementType;
 		}
 		public SummaryTableColumn Clone()
 		{
-			SummaryTableColumn newColumn = new SummaryTableColumn(name, isNumeric, displayName, isRowWeightColumn);
+			SummaryTableColumn newColumn = new SummaryTableColumn(name, isNumeric, displayName, isRowWeightColumn, elementType);
 			newColumn.doubleValues.AddRange(doubleValues);
 			newColumn.stringValues.AddRange(stringValues);
 			newColumn.colourThresholds.AddRange(colourThresholds);
@@ -410,6 +414,22 @@ namespace PerfSummaries
 			return doubleValues[index];
 		}
 
+		public bool AreAllValuesOverThreshold(double threshold)
+		{
+			if (!isNumeric)
+			{
+				return true;
+			}
+			foreach(double value in doubleValues)
+			{
+				if ( value > threshold && value != double.MaxValue)
+				{
+					return true;
+				}
+			}
+			return false;
+		}
+
 		public void SetStringValue(int index, string value)
 		{
 			if (isNumeric)
@@ -537,7 +557,7 @@ namespace PerfSummaries
 			{
 				if (collateByColumns.Contains(srcColumn))
 				{
-					newColumns.Add(new SummaryTableColumn(srcColumn.name, false, srcColumn.displayName));
+					newColumns.Add(new SummaryTableColumn(srcColumn.name, false, srcColumn.displayName, false, srcColumn.elementType));
 					finalSortByList.Add(srcColumn.name.ToLower());
 					// Early out if we've found all the columns
 					if (finalSortByList.Count == collateByColumns.Count)
@@ -547,7 +567,7 @@ namespace PerfSummaries
 				}
 			}
 
-			newColumns.Add(new SummaryTableColumn("Count", true));
+			newColumns.Add(new SummaryTableColumn("Count", true, null, false, SummaryTableElement.Type.ToolMetadata));
 			int countColumnIndex = newColumns.Count - 1;
 
 			int numericColumnStartIndex = newColumns.Count;
@@ -558,11 +578,11 @@ namespace PerfSummaries
 				if (column.isNumeric && !collateByColumns.Contains(column))
 				{
 					srcToDestBaseColumnIndex.Add(newColumns.Count);
-					newColumns.Add(new SummaryTableColumn("Avg " + column.name, true));
+					newColumns.Add(new SummaryTableColumn("Avg " + column.name, true, null, false, column.elementType));
 					if (addMinMaxColumns)
 					{
-						newColumns.Add(new SummaryTableColumn("Min " + column.name, true));
-						newColumns.Add(new SummaryTableColumn("Max " + column.name, true));
+						newColumns.Add(new SummaryTableColumn("Min " + column.name, true, null, false, column.elementType));
+						newColumns.Add(new SummaryTableColumn("Max " + column.name, true, null, false, column.elementType));
 					}
 				}
 				else
@@ -699,7 +719,7 @@ namespace PerfSummaries
 			return SortAndFilter(customFilter.Split(',').ToList(), customRowSort.Split(',').ToList(), bReverseSort, weightByColumnName);
 		}
 
-		public SummaryTable SortAndFilter(List<string> columnFilterList, List<string> rowSortList, bool bReverseSort, string weightByColumnName)
+		public SummaryTable SortAndFilter(List<string> columnFilterList, List<string> rowSortList, bool bReverseSort, string weightByColumnName, float statThreshold = 0.0f)
 		{
 			SummaryTable newTable = SortRows(rowSortList, bReverseSort);
 
@@ -784,6 +804,22 @@ namespace PerfSummaries
 			}
 
 
+			// Filter out csv stat or metric columns below the specified threshold
+			if (statThreshold > 0.0f)
+			{
+				List<SummaryTableColumn> oldColumnList = newColumnList;
+				newColumnList = new List<SummaryTableColumn>();
+				foreach (SummaryTableColumn column in oldColumnList)
+				{
+					if (!column.isNumeric || 
+						( column.elementType != SummaryTableElement.Type.CsvStatAverage && column.elementType != SummaryTableElement.Type.SummaryTableMetric ) || 
+						column.AreAllValuesOverThreshold((double)statThreshold))
+					{
+						newColumnList.Add(column);
+					}
+				}
+			}
+
 			newTable.columns = newColumnList;
 			newTable.rowCount = rowCount;
 			newTable.InitColumnLookup();
@@ -1286,7 +1322,7 @@ namespace PerfSummaries
 			List<SummaryTableColumn> newColumns = new List<SummaryTableColumn>();
 			foreach (SummaryTableColumn srcCol in columns)
 			{
-				SummaryTableColumn destCol = new SummaryTableColumn(srcCol.name, srcCol.isNumeric);
+				SummaryTableColumn destCol = new SummaryTableColumn(srcCol.name, srcCol.isNumeric, null, false, srcCol.elementType);
 				for (int i = 0; i < rowCount; i++)
 				{
 					int srcIndex = columnRemapping[i].Value;
@@ -1339,7 +1375,7 @@ namespace PerfSummaries
 
 				if (!columnLookup.ContainsKey(key))
 				{
-					column = new SummaryTableColumn(value.name, value.isNumeric);
+					column = new SummaryTableColumn(value.name, value.isNumeric, null, false, value.type);
 					columnLookup.Add(key, column);
 					columns.Add(column);
 				}
