Fix CSVCollate args help

diff --git a/Engine/Source/Programs/CSVTools/CSVCollate/CSVCollate.cs b/Engine/Source/Programs/CSVTools/CSVCollate/CSVCollate.cs
--- a/Engine/Source/Programs/CSVTools/CSVCollate/CSVCollate.cs
+++ b/Engine/Source/Programs/CSVTools/CSVCollate/CSVCollate.cs
@@ -14,7 +14,7 @@ namespace CSVTools
 {
     class Version
     {
-        private static string VersionString = "1.32";
+        private static string VersionString = "1.33";
 
         public static string Get() { return VersionString; }
     };
@@ -23,7 +23,7 @@ namespace CSVTools
     {
         static string formatString =
             "Format: \n" +
-            "       -csv <filename>  -csvDir <path>\n" +
+            "       -csvs <filename or ; separated list>  -csvDir <path>\n" +
 			"       [-searchPattern <pattern, e.g *.csv>] - for use with -csvDir\n" +
 			"       [-avg] -stats will be per frame averaged\n" +
 			"       [-recurse] - for use with -csvdir\n" +
