Add -summaryTableXmlSubst method to allow replacing of entries in a summary table. Format is <find1>=<replace1>;<find2>=<replace2>...

diff --git a/Engine/Source/Programs/CSVTools/PerfReportTool/PerfReportTool.cs b/Engine/Source/Programs/CSVTools/PerfReportTool/PerfReportTool.cs
--- a/Engine/Source/Programs/CSVTools/PerfReportTool/PerfReportTool.cs
+++ b/Engine/Source/Programs/CSVTools/PerfReportTool/PerfReportTool.cs
@@ -25,7 +25,7 @@ namespace PerfReportTool
 {
     class Version
     {
-        private static string VersionString = "4.87";
+        private static string VersionString = "4.88";
 
         public static string Get() { return VersionString; }
     };
@@ -160,6 +160,7 @@ namespace PerfReportTool
 			"       -minFrameCount <n> : ignore CSVs without at least this number of valid frames\n" +
 			"       -maxFileAgeDays <n> : max file age in days. CSV or PRC files older than this will be ignored\n" +
 			"       -summaryTableStatThreshold <n> : stat/metric columns in the summarytable will be filtered out if all values < threshold\n" +
+			"       -summaryTableXmlSubst <find1>=<replace1>,<find2>=<replace2>... replace summarytable XML row and filter entries\n" +
 			"\n" +
 			"Json serialization:\n" +
 			"       -summaryTableToJson <filename> : json filename to write summary table row data to\n" +
@@ -333,7 +334,7 @@ namespace PerfReportTool
 			}
 
 			// Load the report + graph XML data
-			reportXML = new ReportXML(GetArg("graphxml", false), GetArg("reportxml", false), GetArg("reportxmlbasedir", false), GetArg("summaryTableXml", false));
+			reportXML = new ReportXML(GetArg("graphxml", false), GetArg("reportxml", false), GetArg("reportxmlbasedir", false), GetArg("summaryTableXml", false), GetArg("summaryTableXmlSubst", false));
 
 			if (GetBoolArg("listSummaryTables"))
 			{
diff --git a/Engine/Source/Programs/CSVTools/PerfReportTool/SummaryTable.cs b/Engine/Source/Programs/CSVTools/PerfReportTool/SummaryTable.cs
--- a/Engine/Source/Programs/CSVTools/PerfReportTool/SummaryTable.cs
+++ b/Engine/Source/Programs/CSVTools/PerfReportTool/SummaryTable.cs
@@ -64,12 +64,13 @@ namespace PerfSummaries
 
 	class SummaryTableInfo
 	{
-		public SummaryTableInfo(XElement tableElement)
+		public SummaryTableInfo(XElement tableElement, Dictionary<string,string> substitutionsDict)
 		{
 			XAttribute rowSortAt = tableElement.Attribute("rowSort");
 			if (rowSortAt != null)
 			{
 				rowSortList.AddRange(rowSortAt.Value.Split(','));
+				ApplySubstitutionsToList(rowSortList, substitutionsDict);
 			}
 			XAttribute weightByColumnAt = tableElement.Attribute("weightByColumn");
 			if (weightByColumnAt != null)
@@ -81,6 +82,7 @@ namespace PerfSummaries
 			if (filterEl != null)
 			{
 				columnFilterList.AddRange(filterEl.Value.Split(','));
+				ApplySubstitutionsToList(columnFilterList, substitutionsDict);
 			}
 
 			bReverseSortRows = tableElement.GetSafeAttibute<bool>("reverseSortRows", false);
@@ -93,8 +95,10 @@ namespace PerfSummaries
 			{
 				if (sectionBoundaryEl != null)
 				{
+					string statName = ApplySubstitution(sectionBoundaryEl.GetSafeAttibute<string>("statName"), substitutionsDict);
+
 					SummarySectionBoundaryInfo sectionBoundary = new SummarySectionBoundaryInfo(
-						sectionBoundaryEl.GetSafeAttibute<string>("statName"),
+						statName,
 						sectionBoundaryEl.GetSafeAttibute<string>("startToken"),
 						sectionBoundaryEl.GetSafeAttibute<string>("endToken"),
 						sectionBoundaryEl.GetSafeAttibute<int>("level", 0),
@@ -106,6 +110,32 @@ namespace PerfSummaries
 			}
 		}
 
+		private void ApplySubstitutionsToList(List<string> list, Dictionary<string, string> substitutionsDict)
+		{
+			if (substitutionsDict == null)
+			{
+				return;
+			}
+			for (int i = 0; i < list.Count; i++)
+			{
+				list[i] = ApplySubstitution(list[i], substitutionsDict);
+			}
+		}
+
+		private string ApplySubstitution(string str, Dictionary<string, string> substitutionsDict)
+		{
+			if (substitutionsDict == null)
+			{
+				return str;
+			}
+			if (substitutionsDict.TryGetValue(str, out string replaceStr))
+			{
+				return replaceStr;
+			}
+			return str;
+		}
+
+
 		public SummaryTableInfo(string filterListStr, string rowSortStr)
 		{
 			columnFilterList.AddRange(filterListStr.Split(','));
