Add missing file from previous submit (21233908)

diff --git a/Engine/Source/Programs/CSVTools/PerfReportTool/ReportXml.cs b/Engine/Source/Programs/CSVTools/PerfReportTool/ReportXml.cs
--- a/Engine/Source/Programs/CSVTools/PerfReportTool/ReportXml.cs
+++ b/Engine/Source/Programs/CSVTools/PerfReportTool/ReportXml.cs
@@ -67,8 +67,9 @@ namespace PerfReportTool
 			return false;
 		}
 
-		public ReportXML(string graphXMLFilenameIn, string reportXMLFilenameIn, string baseXMLDirectoryOverride, string additionalSummaryTableXmlFilename)
+		public ReportXML(string graphXMLFilenameIn, string reportXMLFilenameIn, string baseXMLDirectoryOverride, string additionalSummaryTableXmlFilename, string summaryTableXmlSubstStr)
 		{
+			
 			string location = System.Reflection.Assembly.GetEntryAssembly().Location.ToLower();
 			string baseDirectory = location.Replace("perfreporttool.exe", "");
 
@@ -256,10 +257,27 @@ namespace PerfReportTool
 			summaryTablesElement = rootElement.Element("summaryTables");
 			if (summaryTablesElement != null)
 			{
+				// Read the substitutions
+				Dictionary<string, string> substitutionsDict = null;
+				string[] substitutions = summaryTableXmlSubstStr.Split(',');
+				if (substitutions.Length>0)
+				{
+					substitutionsDict = new Dictionary<string, string>();
+					foreach (string substStr in substitutions)
+					{
+						string [] pair = substStr.Split('=');
+						if (pair.Length == 2)
+						{
+							substitutionsDict[pair[0]] = pair[1];
+						}
+					}
+				}
+
+
 				summaryTables = new Dictionary<string, SummaryTableInfo>();
 				foreach (XElement summaryElement in summaryTablesElement.Elements("summaryTable"))
 				{
-					SummaryTableInfo table = new SummaryTableInfo(summaryElement);
+					SummaryTableInfo table = new SummaryTableInfo(summaryElement, substitutionsDict);
 					summaryTables.Add(summaryElement.Attribute("name").Value.ToLower(), table);
 				}
 			}
