PerfReportTool 4.89 - Transposed summary table and diff row/column support
- Add support for transposed summary tables with -transposeTable or -transposeCollatedTable. Html will be transposed on write.
- Add Diff row support with -addDiffRows. This overrides disables automatic colouring for readability. If -transposeTable is specified, the diff rows become diff columns. Diff text is colorized based on a percentage of the max of the two diffs.
- Various fixes and cleanup for summary table writing. Added HtmlTable helper to assist with HTML table formatting and output.

diff --git a/Engine/Source/Programs/CSVTools/PerfReportTool/PerfReportTool.cs b/Engine/Source/Programs/CSVTools/PerfReportTool/PerfReportTool.cs
--- a/Engine/Source/Programs/CSVTools/PerfReportTool/PerfReportTool.cs
+++ b/Engine/Source/Programs/CSVTools/PerfReportTool/PerfReportTool.cs
@@ -25,7 +25,7 @@ namespace PerfReportTool
 {
     class Version
     {
-        private static string VersionString = "4.88";
+        private static string VersionString = "4.89";
 
         public static string Get() { return VersionString; }
     };
@@ -160,7 +160,10 @@ namespace PerfReportTool
 			"       -minFrameCount <n> : ignore CSVs without at least this number of valid frames\n" +
 			"       -maxFileAgeDays <n> : max file age in days. CSV or PRC files older than this will be ignored\n" +
 			"       -summaryTableStatThreshold <n> : stat/metric columns in the summarytable will be filtered out if all values < threshold\n" +
-			"       -summaryTableXmlSubst <find1>=<replace1>,<find2>=<replace2>... replace summarytable XML row and filter entries\n" +
+			"       -summaryTableXmlSubst <find1>=<replace1>,<find2>=<replace2>... : replace summarytable XML row and filter entries\n" +
+			"       -transposeTable : write the summary tables transposed\n"+
+			"       -transposeCollatedTable : write the collated summary table transposed (disables min/max columns)\n" +
+            "       -addDiffRows : adds diff rows after the first two rows\n" + 
 			"\n" +
 			"Json serialization:\n" +
 			"       -summaryTableToJson <filename> : json filename to write summary table row data to\n" +
@@ -714,6 +717,9 @@ namespace PerfReportTool
 				weightByColumnName = null;
 			}
 
+			bool bTransposeFullSummaryTable = GetBoolArg("transposeTable");
+			bool bTransposeCollatedSummaryTable = bTransposeFullSummaryTable | GetBoolArg("transposeCollatedTable");
+
 			// Check params and any commandline overrides
 			bool bReverseTable = tableInfo.bReverseSortRows;
 			bool? bReverseTableOption = GetOptionalBoolArg("reverseTable");
@@ -736,7 +742,7 @@ namespace PerfReportTool
 				bAutoColorizeTable = (bool)bAutoColorizeTableOption;
 			}
 
-			bool addMinMaxColumns = !GetBoolArg("noSummaryMinMax");
+			bool addMinMaxColumns = !GetBoolArg("noSummaryMinMax") && !bTransposeCollatedSummaryTable;
 
 			if (!string.IsNullOrEmpty(outputDir))
 			{
@@ -758,6 +764,12 @@ namespace PerfReportTool
 				filteredTable.ApplyDisplayNameMapping(statDisplaynameMapping);
 				string VersionString = GetBoolArg("noWatermarks") ? "" : Version.Get();
 				string summaryTitle = GetArg("summaryTitle", null);
+
+				if (GetBoolArg("addDiffRows"))
+				{
+					filteredTable.AddDiffRows();
+				}
+
 				filteredTable.WriteToHTML(
 					filenameWithoutExtension + ".html", 
 					VersionString, 
@@ -770,7 +782,9 @@ namespace PerfReportTool
 					GetIntArg("maxSummaryTableStringLength", Int32.MaxValue), 
 					reportXML.columnFormatInfoList, 
 					weightByColumnName, 
-					summaryTitle );
+					summaryTitle,
+					bCollated ? bTransposeCollatedSummaryTable : bTransposeFullSummaryTable
+				);
 			}
 		}
 
diff --git a/Engine/Source/Programs/CSVTools/PerfReportTool/SummaryTable.cs b/Engine/Source/Programs/CSVTools/PerfReportTool/SummaryTable.cs
--- a/Engine/Source/Programs/CSVTools/PerfReportTool/SummaryTable.cs
+++ b/Engine/Source/Programs/CSVTools/PerfReportTool/SummaryTable.cs
@@ -13,6 +13,7 @@ namespace PerfSummaries
 	{
 		public static string FormatStatName(string inStatName)
 		{
+			// We add a space so stats are split over multiple lines
 			return inStatName.Replace("/", "/ ");
 		}
 
@@ -256,13 +257,16 @@ namespace PerfSummaries
 		public bool isNumeric = false;
 		public string displayName;
 		public bool isRowWeightColumn = false;
+		public bool hasDiffRows = false;
 		List<double> doubleValues = new List<double>();
 		List<string> stringValues = new List<string>();
 		List<string> toolTips = new List<string>();
 		public SummaryTableElement.Type elementType;
+		public SummaryTableColumnFormatInfo formatInfo = null;
 
 		List<ColourThresholdList> colourThresholds = new List<ColourThresholdList>();
 		ColourThresholdList colourThresholdOverride = null;
+
 		public SummaryTableColumn(string inName, bool inIsNumeric, string inDisplayName, bool inIsRowWeightColumn, SummaryTableElement.Type inElementType)
 		{
 			name = inName;
@@ -278,37 +282,111 @@ namespace PerfSummaries
 			newColumn.stringValues.AddRange(stringValues);
 			newColumn.colourThresholds.AddRange(colourThresholds);
 			newColumn.toolTips.AddRange(toolTips);
+			newColumn.hasDiffRows = hasDiffRows;
 			return newColumn;
 		}
 
-		public string GetDisplayName(string hideStatPrefix=null)
+		private double FilterInvalidValue(double value)
+		{
+			return value == double.MaxValue ? 0.0 : value;
+		}
+
+		public void AddDiffRows(bool bIsFirstColumn=false)
 		{
-			string nameOut;
-			if (displayName == null)
+			if (hasDiffRows)
 			{
-				// Trim the stat name suffix if necessary
-				string statName = name;
-				if (hideStatPrefix != null)
+				throw new Exception("Column already has diff rows!");
+			}
+			// Add a diff row for every row after the first one
+			int oldCount = GetCount();
+			int diffRowCount = oldCount - 1;
+			int newCount = oldCount + diffRowCount;
+
+			// Create new lists with counts reserved
+			List<double> newDoubleValues = new List<double>(doubleValues.Count > 0 ? newCount : 0);
+			List<string> newStringValues = new List<string>(stringValues.Count > 0 ? newCount : 0);
+			List<string> newToolTips = new List<string>(toolTips.Count > 0 ? newCount : 0);
+			List<ColourThresholdList> newColourThresholds = new List<ColourThresholdList>(colourThresholds.Count > 0 ? newCount : 0);
+
+			// Add diff rows to each of the arrays
+			for (int i = 0; i < doubleValues.Count; i++)
+			{
+				newDoubleValues.Add(doubleValues[i]);
+				if (i > 0)
 				{
-					string prefix = "";
-					string suffix = name;
-					if (name.StartsWith("Min ") || name.StartsWith("Max ") || name.StartsWith("Avg "))
-					{
-						prefix = name.Substring(0, 4);
-						suffix = name.Substring(4);
-					}
-					if ( suffix.ToLower().StartsWith(hideStatPrefix.ToLower() ) )
-					{
-						statName = prefix + suffix.Substring(hideStatPrefix.Length);
-					}
+					double thisValue = FilterInvalidValue( doubleValues[i] );
+					double prevValue = FilterInvalidValue( doubleValues[i - 1] );
+					newDoubleValues.Add( thisValue - prevValue );
 				}
-				nameOut = TableUtil.FormatStatName(statName);
 			}
-			else
+			for (int i = 0; i < stringValues.Count; i++)
+			{
+				newStringValues.Add(stringValues[i]);
+				if (i > 0)
+				{
+					newStringValues.Add(bIsFirstColumn ? "Diff" : "");
+				}
+			}
+			for (int i = 0; i < toolTips.Count; i++)
+			{
+				newToolTips.Add(toolTips[i]);
+				if (i > 0)
+				{
+					newToolTips.Add("");
+				}
+			}
+			for (int i = 0; i < colourThresholds.Count; i++)
+			{
+				newColourThresholds.Add(colourThresholds[i]);
+				if (i > 0)
+				{
+					newColourThresholds.Add(null);
+				}
+			}
+
+			doubleValues = newDoubleValues;
+			stringValues = newStringValues;
+			toolTips = newToolTips;
+			colourThresholds = newColourThresholds;
+			hasDiffRows = true;
+		}
+
+		public bool IsDiffRow(int rowIndex)
+		{
+			if (hasDiffRows == false || rowIndex < 2)
+			{
+				return false;
+			}
+			return ((rowIndex - 2) % 2) == 0;
+		}
+
+		public string GetDisplayName(string hideStatPrefix=null, bool bAddSeparatorSpaces = true)
+		{
+			if (displayName != null)
 			{
-				nameOut = displayName;
+				return displayName;
 			}
-			return nameOut;
+			// Trim the stat name suffix if necessary
+			string statName = name;
+			if (hideStatPrefix != null)
+			{
+				string prefix = "";
+				string suffix = name;
+				if (name.StartsWith("Min ") || name.StartsWith("Max ") || name.StartsWith("Avg "))
+				{
+					prefix = name.Substring(0, 4);
+					suffix = name.Substring(4);
+				}
+				if ( suffix.ToLower().StartsWith(hideStatPrefix.ToLower() ) )
+				{
+					statName = prefix + suffix.Substring(hideStatPrefix.Length);
+				}
+			}
+			if (bAddSeparatorSpaces)
+			{
+				return TableUtil.FormatStatName(statName);
+			}
+			return statName;
 		}
 
 		public void SetValue(int index, double value)
@@ -366,7 +444,7 @@ namespace PerfSummaries
 			return null;
 		}
 
-		public string GetColour(int index)
+		public string GetBackgroundColor(int index)
 		{
 			ColourThresholdList thresholds = null;
 			double value = GetValue(index);
@@ -392,43 +470,85 @@ namespace PerfSummaries
 			return thresholds.GetColourForValue(value);
 		}
 
-		public void ComputeAutomaticColourThresholds(AutoColorizeMode autoColorizeMode)
+		public string GetTextColor(int index)
 		{
-			if (autoColorizeMode == AutoColorizeMode.Off)
+			const double absoluteIgnoreThreshold = 0.05;
+			if (hasDiffRows && IsDiffRow(index) && isNumeric && formatInfo.autoColorizeMode != AutoColorizeMode.Off && index < doubleValues.Count )
 			{
-				return;
-			}
-			colourThresholds = new List<ColourThresholdList>();
-			double maxValue = -double.MaxValue;
-			double minValue = double.MaxValue;
-			double totalValue = 0.0f;
-			double validCount = 0.0f;
-			for (int i = 0; i < doubleValues.Count; i++)
-			{
-				double val = doubleValues[i];
-				if (val != double.MaxValue)
+				double diffValue = doubleValues[index];
+				if ( Math.Abs(diffValue) < absoluteIgnoreThreshold )
+				{
+					return "#A0A0A0";
+				}
+				double prevValue = FilterInvalidValue(doubleValues[index - 2]);
+				double thisValue = FilterInvalidValue(doubleValues[index - 1]);
+
+				double maxValue = Math.Max(thisValue, prevValue);
+				double percentOfMax = 100.0 * diffValue / maxValue;
+				string red = "#B00000";
+				string green = "#008000";
+				if (percentOfMax >= 1.0)
+				{
+					return (formatInfo.autoColorizeMode == AutoColorizeMode.HighIsBad) ? red : green;
+				}
+				if (percentOfMax <= -1.0)
 				{
-					maxValue = Math.Max(val, maxValue);
-					minValue = Math.Min(val, minValue);
-					totalValue += val;
-					validCount += 1.0f;
+					return (formatInfo.autoColorizeMode == AutoColorizeMode.HighIsBad) ? green : red;
 				}
+				return "#808080";
 			}
-			if (minValue == maxValue)
+			return null;
+		}
+
+
+		public void ComputeAutomaticColourThresholds()
+		{
+			AutoColorizeMode autoColorizeMode = formatInfo.autoColorizeMode;
+			if (autoColorizeMode == AutoColorizeMode.Off || !isNumeric)
 			{
 				return;
 			}
 
-			Colour green = new Colour(0.4f, 0.82f, 0.45f);
-			Colour yellow = new Colour(1.0f, 1.0f, 0.5f);
-			Colour red = new Colour(1.0f, 0.4f, 0.4f);
+			if (hasDiffRows)
+			{
+				// If we have diff rows then color thresholds will be blank. We'll use font colors instead
+				colourThresholds = new List<ColourThresholdList>();
+			}
+			else // !hasDiffRows
+			{
+				// Set a single colour threshold list for the whole column
+				colourThresholds = new List<ColourThresholdList>();
+				double maxValue = -double.MaxValue;
+				double minValue = double.MaxValue;
+				double totalValue = 0.0f;
+				double validCount = 0.0f;
+				for (int i = 0; i < doubleValues.Count; i++)
+				{
+					double val = doubleValues[i];
+					if (val != double.MaxValue)
+					{
+						maxValue = Math.Max(val, maxValue);
+						minValue = Math.Min(val, minValue);
+						totalValue += val;
+						validCount += 1.0f;
+					}
+				}
+				if (minValue == maxValue)
+				{
+					return;
+				}
+
+				Colour green = new Colour(0.4f, 0.82f, 0.45f);
+				Colour yellow = new Colour(1.0f, 1.0f, 0.5f);
+				Colour red = new Colour(1.0f, 0.4f, 0.4f);
 
-			double averageValue = totalValue / validCount; // TODO: Weighted average 
-			colourThresholdOverride = new ColourThresholdList();
-			colourThresholdOverride.Add(new ThresholdInfo(minValue, (autoColorizeMode == AutoColorizeMode.HighIsBad) ? green : red));
-			colourThresholdOverride.Add(new ThresholdInfo(averageValue, yellow));
-			colourThresholdOverride.Add(new ThresholdInfo(averageValue, yellow));
-			colourThresholdOverride.Add(new ThresholdInfo(maxValue, (autoColorizeMode == AutoColorizeMode.HighIsBad) ? red : green));
+				double averageValue = totalValue / validCount; // TODO: Weighted average 
+				colourThresholdOverride = new ColourThresholdList();
+				colourThresholdOverride.Add(new ThresholdInfo(minValue, (autoColorizeMode == AutoColorizeMode.HighIsBad) ? green : red));
+				colourThresholdOverride.Add(new ThresholdInfo(averageValue, yellow));
+				colourThresholdOverride.Add(new ThresholdInfo(averageValue, yellow));
+				colourThresholdOverride.Add(new ThresholdInfo(maxValue, (autoColorizeMode == AutoColorizeMode.HighIsBad) ? red : green));
+			}
 		}
 
 		public int GetCount()
@@ -487,9 +607,24 @@ namespace PerfSummaries
 					return "";
 				}
 				double val = doubleValues[index];
+
+				string prefix = "";
+				bool bIsDiffRow = hasDiffRows && IsDiffRow(index);
+				if (bIsDiffRow)
+				{
+					if ( val == 0.0 )
+					{
+						return "";
+					}
+					if ( val > 0.0 )
+					{
+						prefix = "+";
+					}
+				}
+
 				if (forceNumericFormat != null)
 				{
-					return val.ToString(forceNumericFormat);
+					return prefix + val.ToString(forceNumericFormat);
 				}
 				else if (roundNumericValues)
 				{
@@ -497,19 +632,31 @@ namespace PerfSummaries
 					double frac = absVal - (double)Math.Truncate(absVal);
 					if (absVal >= 250.0f || frac < 0.0001f)
 					{
-						return val.ToString("0");
+						return prefix+val.ToString("0");
 					}
 					if (absVal >= 50.0f)
 					{
-						return val.ToString("0.0");
+						return prefix + val.ToString("0.0");
 					}
 					if (absVal >= 0.1)
 					{
-						return val.ToString("0.00");
+						return prefix + val.ToString("0.00");
+					}
+					if (bIsDiffRow)
+					{
+						string strOut = val.ToString("0.000");
+						if (strOut == "0.000")
+						{
+							return "";
+						}
+						return prefix + strOut;
+					}
+					else
+					{
+						return val.ToString("0.000");
 					}
-					return val.ToString("0.000");
 				}
-				return val.ToString();
+				return prefix + val.ToString();
 			}
 			else
 			{
@@ -857,6 +1004,18 @@ namespace PerfSummaries
 			return newTable;
 		}
 
+		public void AddDiffRows()
+		{
+			for (int i=0; i<columns.Count; i++)
+			{
+				columns[i].AddDiffRows(i == 0);
+			}
+			rowCount = rowCount + (rowCount - 1);
+
+			// Just set rowWeightings to null for now. We shouldn't need it, since we should have already collated by this point
+			rowWeightings = null;
+		}
+
 		public void ApplyDisplayNameMapping(Dictionary<string, string> statDisplaynameMapping)
 		{
 			// Convert to a display-friendly name
@@ -939,7 +1098,8 @@ namespace PerfSummaries
 			int maxColumnStringLength, 
 			SummaryTableColumnFormatInfoCollection columnFormatInfoList, 
 			string weightByColumnName, 
-			string title)
+			string title, 
+			bool bTranspose )
 		{
 			System.IO.StreamWriter htmlFile = new System.IO.StreamWriter(htmlFilename, false);
 			int statColSpan = hasMinMaxColumns ? 3 : 1;
@@ -959,12 +1119,14 @@ namespace PerfSummaries
 			htmlFile.WriteLine("<html>");
 			htmlFile.WriteLine("<head><title>Perf Summary: "+ title + "</title>");
 
+			bool bAddStatNameSpacing = true;// !bTranspose;
+
 			// Figure out the sticky column count
 			int stickyColumnCount = 0;
 			if (bScrollableTable)
 			{
 				stickyColumnCount = 1;
-				if (isCollated)
+				if (isCollated && !bTranspose)
 				{
 					for (int i = 0; i < columns.Count; i++)
 					{
@@ -978,10 +1140,9 @@ namespace PerfSummaries
 			}
 
 			// Get format info for the columns
-			Dictionary<SummaryTableColumn, SummaryTableColumnFormatInfo> columnFormatInfoLookup = new Dictionary<SummaryTableColumn, SummaryTableColumnFormatInfo>();
 			foreach (SummaryTableColumn column in columns)
 			{
-				columnFormatInfoLookup[column] = (columnFormatInfoList != null) ? columnFormatInfoList.GetFormatInfo(column.name) : SummaryTableColumnFormatInfoCollection.DefaultColumnInfo;
+				column.formatInfo = (columnFormatInfoList != null) ? columnFormatInfoList.GetFormatInfo(column.name) : SummaryTableColumnFormatInfoCollection.DefaultColumnInfo;
 			}
 
 			// Automatically colourize the table if requested
@@ -991,7 +1152,7 @@ namespace PerfSummaries
 				{
 					if (column.isNumeric)
 					{
-						column.ComputeAutomaticColourThresholds(columnFormatInfoLookup[column].autoColorizeMode);
+						column.ComputeAutomaticColourThresholds();
 					}
 				}
 			}
@@ -1028,17 +1189,9 @@ namespace PerfSummaries
 					"	 styleElement.textContent+='tr.lastHeaderRow th:nth-child('+(i+1)+') {  z-index: 8;  border-top: 2px solid black;  font-size: 11px;  left: '+xOffset+'px;'+rBorderParam+'}'; \n" +
 					"	 styleElement.textContent+='td:nth-child('+(i+1)+') {  position: -webkit-sticky;  position: sticky; z-index: 7;  left: '+xOffset+'px; '+rBorderParam+'}'; \n" +
 					"	 xOffset+=table.rows[1].cells[i].offsetWidth; \n" +
-					"  } \n"
-					);
-				// Make all sticky columns except the count column opaque
-				htmlFile.WriteLine(
-					"  for (var i=0;i<numStickyCols-1;i++) \n" +
-					"  { \n" +
-					"    styleElement.textContent+='tr:nth-child(odd) td:nth-child('+(i+1)+') { background-color: #e2e2e2; }'; \n" +
-					"    styleElement.textContent+='tr:nth-child(even) td:nth-child('+(i+1)+') { background-color: #ffffff; }'; \n" +
-					"  } \n" +
+					"  } \n"+
 					"} \n"
-				);
+					);
 				htmlFile.WriteLine("</script>");
 			}
 
@@ -1048,6 +1201,8 @@ namespace PerfSummaries
 			htmlFile.WriteLine("h2 {  font-family: 'Verdana', Times, serif; font-size: 16px }");
 			htmlFile.WriteLine("h1 {  font-family: 'Verdana', Times, serif; font-size: 20px }");
 			string tableCss = "";
+
+
 			if (bScrollableTable)
 			{
 				int firstColMaxStringLength = 0;
@@ -1058,18 +1213,30 @@ namespace PerfSummaries
 						firstColMaxStringLength = Math.Max(firstColMaxStringLength, columns[0].GetStringValue(i).Length);
 					}
 				}
+
+				int headerMinWidth = bTranspose ? 50 : 75;
+				int cellFontSize = bTranspose ? 11 : 10;
+				int headerCellFontSize = cellFontSize - 1;
+
+				if (bAddMinMaxColumns)
+				{
+					headerMinWidth = 40;
+				}
+
 				int firstColWidth = (int)(firstColMaxStringLength * 6.5);
+
 				tableCss =
 					"table {table-layout: fixed;} \n" +
-					"table, th, td { border: 0px solid black; border-spacing: 0; border-collapse: separate; padding: " + cellPadding + "px; vertical-align: center; font-family: 'Verdana', Times, serif; font-size: 10px;} \n" +
+					"table, th, td { border: 0px solid black; border-spacing: 0; border-collapse: separate; padding: " + cellPadding + "px; vertical-align: center; font-family: 'Verdana', Times, serif; font-size: "+ cellFontSize + "px;} \n" +
 					"td {" +
 					"  border-right: 1px solid black;" +
-					"  max-width: 400;" +
+					"  max-width: 450;" +
 					"} \n" +
 					"tr:first-element { border-top: 2px; border-bottom: 2px } \n" +
 					"th {" +
-					"  width: 75px;" +
-					"  max-width: 400;" +
+					"  width: auto;"+
+					"  max-width: 220px;" +
+					"  min-width: "+ headerMinWidth + "px;" +
 					"  position: -webkit-sticky;" +
 					"  position: sticky;" +
 					"  border-right: 1px solid black;" +
@@ -1077,7 +1244,7 @@ namespace PerfSummaries
 					"  z-index: 5;" +
 					"  background-color: #ffffff;" +
 					"  top:0;" +
-					"  font-size: 9px;" +
+					"  font-size: "+ headerCellFontSize + "px;" +
 					"  word-wrap: break-word;" +
 					"  overflow: hidden;" +
 					"  height: 60;" +
@@ -1109,8 +1276,13 @@ namespace PerfSummaries
 
 
 			bool bOddRowsGray = !(!bAddMinMaxColumns || !isCollated);
-			tableCss += "tr:nth-child(" + (bOddRowsGray ? "odd" : "even") + ") {background-color: #e2e2e2;} \n";
-			tableCss += "tr:nth-child(" + (bOddRowsGray ? "even" : "odd") + ") {background-color: #ffffff;} \n";
+
+			string oddColor = bOddRowsGray ? "#e2e2e2" : "#ffffff";
+			string evenColor = bOddRowsGray ? "#ffffff" : "#e2e2e2";
+
+			tableCss += "tr:nth-child(odd) {background-color: "+oddColor+";} \n";
+			tableCss += "tr:nth-child(even) {background-color: "+evenColor+";} \n";
+
 			tableCss += "tr:first-child {background-color: #ffffff;} \n";
 			tableCss += "tr.lastHeaderRow th { border-bottom: 2px solid black; } \n";
 
@@ -1123,60 +1295,65 @@ namespace PerfSummaries
 
 			htmlFile.WriteLine("</style>");
 			htmlFile.WriteLine("</head><body>");
-			htmlFile.WriteLine("<table id='mainTable'>");
 
-			string HeaderRow = "";
+
+			HtmlTable htmlTable = new HtmlTable("id='mainTable'", 1, columns.Count);
+
+			// Make a header row, but don't add it to the table yet
+			HtmlTable.Row headerRow = new HtmlTable.Row();
+
+			// Write the header
 			if (isCollated)
 			{
-				string TopHeaderRow = "";
-				if (bScrollableTable)
-				{
-					TopHeaderRow += "<th colspan='" + firstStatColumnIndex + "'><h3>" + title + "</h3></th>";
-				}
-				else
-				{
-					TopHeaderRow += "<th colspan='" + firstStatColumnIndex + "'/>";
-				}
+				HtmlTable.Row topHeaderRow = headerRow;
 
+				// Add the special columns (up to Count) to the lower header row
 				for (int i = 0; i < firstStatColumnIndex; i++)
 				{
-					HeaderRow += "<th>" + columns[i].GetDisplayName(hideStatPrefix) + "</th>";
-				}
-				if (!bAddMinMaxColumns)
-				{
-					TopHeaderRow = HeaderRow;
+					headerRow.AddCell(columns[i].GetDisplayName(hideStatPrefix, bAddStatNameSpacing));
 				}
 
-				for (int i = firstStatColumnIndex; i < columns.Count; i++)
+				if (bAddMinMaxColumns)
 				{
-					string prefix = "";
-					string suffix = "";
-					string statName = GetStatNameWithPrefixAndSuffix(columns[i].GetDisplayName(hideStatPrefix), out prefix, out suffix);
-					if ((i - 1) % statColSpan == 0)
+					topHeaderRow = htmlTable.CreateRow();
+					htmlTable.numHeaderRows = 2;
+					if (bScrollableTable)
 					{
-						TopHeaderRow += "<th colspan='" + statColSpan + "' >" + statName + suffix + "</th>";
+						topHeaderRow.AddCell("<h3>" + title + "</h3>", "colspan='" + firstStatColumnIndex + "'");
+					}
+					else
+					{
+						topHeaderRow.AddCell("", "colspan='" + firstStatColumnIndex + "'");
+					}
+					// Add the stat columns
+					for (int i = firstStatColumnIndex; i < columns.Count; i++)
+					{
+						string statName = GetStatNameWithPrefixAndSuffix(columns[i].GetDisplayName(hideStatPrefix, bAddStatNameSpacing), out string prefix, out string suffix);
+						if ((i - 1) % statColSpan == 0)
+						{
+							topHeaderRow.AddCell(statName + suffix, "colspan='" + statColSpan + "'");
+						}
+						headerRow.AddCell(prefix.Trim());
 					}
-					HeaderRow += "<th>" + prefix.Trim() + "</th>";
-				}
-				if (bAddMinMaxColumns)
-				{
-					htmlFile.WriteLine("  <tr>" + TopHeaderRow + "</tr>");
-					htmlFile.WriteLine("  <tr class='lastHeaderRow'>" + HeaderRow + "</tr>");
 				}
 				else
 				{
-					htmlFile.WriteLine("  <tr class='lastHeaderRow'>" + TopHeaderRow + "</tr>");
+					// Add the stat columns
+					for (int i = firstStatColumnIndex; i < columns.Count; i++)
+					{
+						string statName = GetStatNameWithPrefixAndSuffix(columns[i].GetDisplayName(hideStatPrefix, bAddStatNameSpacing), out _, out string suffix);
+						headerRow.AddCell(statName + suffix);
+					}
 				}
 			}
 			else
 			{
 				foreach (SummaryTableColumn column in columns)
 				{
-					HeaderRow += "<th>" + column.GetDisplayName(hideStatPrefix) + "</th>";
+					headerRow.AddCell(column.GetDisplayName(hideStatPrefix, bAddStatNameSpacing));
 				}
-				htmlFile.WriteLine("  <tr class='lastHeaderRow'>" + HeaderRow + "</tr>");
 			}
-			string[] stripeColors = { "'#e2e2e2'", "'#ffffff'" };
+			htmlTable.AddRow(headerRow);
 
 			// Work out which rows are major/minor section boundaries
 			Dictionary<int, int> rowSectionBoundaryLevel = new Dictionary<int, int>();
@@ -1246,54 +1423,57 @@ namespace PerfSummaries
 						}
 					}
 				}
-
 			}
 
 			// Add the rows to the table
-			for (int i = 0; i < rowCount; i++)
+			for (int rowIndex = 0; rowIndex < rowCount; rowIndex++)
 			{
 				string rowClassStr = "";
 
 				// Is this a major/minor section boundary
-				if (rowSectionBoundaryLevel.ContainsKey(i))
+				if (rowSectionBoundaryLevel.ContainsKey(rowIndex))
 				{
-					int sectionLevel = rowSectionBoundaryLevel[i];
+					int sectionLevel = rowSectionBoundaryLevel[rowIndex];
 					if (sectionLevel < 3)
 					{
 						rowClassStr = " class='sectionStartLevel" + sectionLevel + "'";
 					}
 				}
 
-				htmlFile.Write("<tr" + rowClassStr + ">");
+				HtmlTable.Row currentRow = htmlTable.CreateRow(rowClassStr);
 				int columnIndex = 0;
 				foreach (SummaryTableColumn column in columns)
 				{
+					List<string> attributes = new List<string>();
+
 					// Add the tooltip for non-collated tables
-					string toolTipString = "";
 					if (!isCollated)
 					{
-						string toolTip = column.GetToolTipValue(i);
+						string toolTip = column.GetToolTipValue(rowIndex);
 						if (toolTip == "")
 						{
-							toolTip = column.GetDisplayName();
+							toolTip = column.GetDisplayName(null, bAddStatNameSpacing);
 						}
-						toolTipString = " title='" + toolTip + "'";
+						attributes.Add("title='" + toolTip + "'");
 					}
-					string colour = column.GetColour(i);
+					string bgColour = column.GetBackgroundColor(rowIndex);
 
-					// Alternating row colours are normally handled by CSS, but we need to handle it explicitly if we have frozen first columns
-					if (columnIndex < stickyColumnCount && colour == null)
+					if (bgColour != null)
 					{
-						colour = stripeColors[i % 2];
+						attributes.Add("bgcolor=" + bgColour);
+					}
+					string textColour = column.GetTextColor(rowIndex);
+					if (textColour != null)
+					{
+						attributes.Add("style='color:" + textColour + "'");
 					}
-					string bgColorString = (colour == null ? "" : " bgcolor=" + colour);
-					bool bold = false;
 
-					SummaryTableColumnFormatInfo columnFormat = columnFormatInfoLookup[column];
+					bool bold = false;
+					SummaryTableColumnFormatInfo columnFormat = column.formatInfo;
 					int maxStringLength = Math.Min( isCollated ? columnFormat.maxStringLengthCollated : columnFormat.maxStringLength, maxColumnStringLength);
 
 					string numericFormat = columnFormat.numericFormat;
-					string stringValue = column.GetStringValue(i, true, numericFormat);
+					string stringValue = column.GetStringValue(rowIndex, true, numericFormat);
 					if (stringValue.Length > maxStringLength)
 					{
 						stringValue = TableUtil.SafeTruncateHtmlTableValue(stringValue, maxStringLength);
@@ -1302,18 +1482,44 @@ namespace PerfSummaries
 					{
 						stringValue = "'" + stringValue;
 					}
-					string columnString = "<td" + toolTipString + bgColorString + "> " + (bold ? "<b>" : "") + stringValue + (bold ? "</b>" : "") + "</td>";
-					htmlFile.Write(columnString);
+					currentRow.AddCell( (bold ? "<b>" : "") + stringValue + (bold ? "</b>" : ""), attributes );
 					columnIndex++;
 				}
-				htmlFile.WriteLine("</tr>");
 			}
-			htmlFile.WriteLine("</table>");
+
+			if (bTranspose)
+			{
+				htmlTable = htmlTable.Transpose();
+				htmlTable.Set(0, 0, new HtmlTable.Cell(title, ""));
+			}
+
+			// Apply final formatting
+			htmlTable.rows[htmlTable.numHeaderRows-1].attributes += " class='lastHeaderRow'";
+
+			if (bScrollableTable)
+			{
+				// Apply stripe colors to the sticky columns to make them opaque (these need to render on top).	Can't use per-cell CSS because we don't want to override existing cell colors
+				string[] stripeColors = { "bgcolor='"+oddColor+"'", "bgcolor='"+evenColor+"'" };
+				for (int rowIndex = htmlTable.numHeaderRows; rowIndex < htmlTable.rows.Count; rowIndex++)
+				{
+					HtmlTable.Row row = htmlTable.rows[rowIndex];
+					for (int colIndex = 0; colIndex<Math.Min(row.cells.Count,stickyColumnCount); colIndex++ )
+					{
+						HtmlTable.Cell cell = row.cells[colIndex];
+						if (cell != null && !cell.attributes.Contains("bgcolor="))
+						{
+							cell.attributes += " " + stripeColors[rowIndex % 2];
+						}
+					}
+				}
+			}
+
+			htmlTable.WriteToHtml(htmlFile);
+
 			string extraString = "";
 			if (isCollated && weightByColumnName != null)
 			{
 				extraString += " - weighted avg";
-				//htmlFile.WriteLine("<p style='font-size:8'>Weighted by " + weightByColumnName +"</p>");
 			}
 
 			htmlFile.WriteLine("<p style='font-size:8'>Created with PerfReportTool " + VersionString + extraString + "</p>");
@@ -1442,4 +1648,164 @@ namespace PerfSummaries
 		bool hasMinMaxColumns = false;
 	};
 
+	class HtmlTable
+	{
+		public class Row
+		{
+			public Row(string inAttributes="", int reserveColumnCount=0)
+			{
+				attributes = inAttributes;
+				if (reserveColumnCount > 0)
+				{
+					cells = new List<Cell>(reserveColumnCount);
+				}
+				else
+				{
+					cells = new List<Cell>();
+				}
+			}
+
+			public void AddCell(string contents="", string attributes="")
+			{
+				cells.Add(new Cell(contents, attributes));
+			}
+
+			public void AddCell(string contents, List<string> attributes)
+			{
+				cells.Add(new Cell(contents, String.Join(" ",attributes)));
+			}
+
+
+			public void WriteToHtml(System.IO.StreamWriter htmlFile, bool bIsHeaderRow )
+			{
+				string attributesStr = attributes.Length == 0 ? "" : " " + attributes;
+				htmlFile.Write("<tr"+ attributesStr + ">");
+				foreach (Cell cell in cells)
+				{
+					if (cell == null)
+					{
+						htmlFile.Write("<td/>");
+					}
+					else
+					{
+						cell.WriteToHtml(htmlFile, bIsHeaderRow);
+					}
+				}
+				htmlFile.WriteLine("</tr>");
+			}
+
+			public void Set(int ColIndex, Cell cell)
+			{
+				if (cells.Count < ColIndex + 1)
+				{
+					cells.Capacity = ColIndex + 1;
+					while (cells.Count < ColIndex + 1)
+					{
+						cells.Add(null);
+					}
+				}
+				cells[ColIndex] = cell;
+			}
+
+			public string attributes;
+			public List<Cell> cells;
+		}
+
+		public class Cell
+		{
+			public Cell()
+			{
+			}
+			public Cell(Cell srcCell)
+			{
+				contents = srcCell.contents;
+				attributes = srcCell.attributes;
+			}
+
+			public Cell(string inContents, string inAttributes, bool bInIsHeader=false)
+			{
+				contents = inContents;
+				attributes = inAttributes;
+			}
+
+			public void WriteToHtml(System.IO.StreamWriter htmlFile, bool bIsHeaderRow)
+			{
+				string AttributesStr = attributes.Length == 0 ? "" : " "+ attributes;
+				string CellType = bIsHeaderRow ? "th" : "td";
+				htmlFile.Write("<" + CellType + AttributesStr + ">" + contents + "</" + CellType + ">");
+			}
+
+			public string contents;
+			public string attributes;
+		}
+
+		public HtmlTable(string inAttributes="", int inNumHeaderRows=1, int reserveColumnCount=0)
+		{
+			rows = new List<Row>();
+			attributes = inAttributes;
+			numHeaderRows = inNumHeaderRows;
+			reserveColumnCount = 0;
+		}
+
+		public HtmlTable Transpose()
+		{
+			// NOTE: Row attributes are stripped when transposing!
+			HtmlTable transposedTable = new HtmlTable(attributes, 1, rows.Count);
+			for (int i=0; i<rows.Count; i++)
+			{
+				Row row = rows[i];
+				for (int j=0; j< row.cells.Count; j++)
+				{
+					Cell newCell = new Cell(row.cells[j]);
+					transposedTable.Set(j, i, newCell);
+				}
+			}
+			return transposedTable;
+		}
+
+		public void WriteToHtml(System.IO.StreamWriter htmlFile)
+		{
+			string attributesStr = attributes.Length == 0 ? "" : " " + attributes;
+			htmlFile.WriteLine("<table"+attributesStr+">");
+
+			// Compute the number of columns, just in case not all rows are equal
+			for (int i=0; i<rows.Count; i++)
+			{
+				rows[i].WriteToHtml(htmlFile, i<numHeaderRows);
+			}
+			htmlFile.WriteLine("</table>");
+		}
+
+		public void Set(int RowIndex, int ColIndex, Cell cell)
+		{
+			// Make sure we have enough rows
+			if (rows.Count < RowIndex + 1)
+			{
+				rows.Capacity = RowIndex + 1;
+				while (rows.Count < RowIndex + 1)
+				{
+					rows.Add(new Row("", reserveColumnCount));
+				}
+			}
+			// Make sure the row is big enough
+			rows[RowIndex].Set(ColIndex, cell);
+		}
+
+		public Row CreateRow(string attributes="")
+		{
+			Row newRow = new Row(attributes, reserveColumnCount);
+			rows.Add(newRow);
+			return newRow; 
+		}
+		public void AddRow(Row row)
+		{
+			rows.Add(row);
+		}
+
+		public List<Row> rows;
+		public string attributes;
+		public int numHeaderRows = 0;
+		int reserveColumnCount = 0;
+	};
+
 }
\ No newline at end of file
