Reduce thresholds for diff column colourization

diff --git a/Engine/Source/Programs/CSVTools/PerfReportTool/SummaryTable.cs b/Engine/Source/Programs/CSVTools/PerfReportTool/SummaryTable.cs
--- a/Engine/Source/Programs/CSVTools/PerfReportTool/SummaryTable.cs
+++ b/Engine/Source/Programs/CSVTools/PerfReportTool/SummaryTable.cs
@@ -472,7 +472,7 @@ namespace PerfSummaries
 
 		public string GetTextColor(int index)
 		{
-			const double absoluteIgnoreThreshold = 0.05;
+			const double absoluteIgnoreThreshold = 0.025;
 			if (hasDiffRows && IsDiffRow(index) && isNumeric && formatInfo.autoColorizeMode != AutoColorizeMode.Off && index < doubleValues.Count )
 			{
 				double diffValue = doubleValues[index];
@@ -487,11 +487,11 @@ namespace PerfSummaries
 				double percentOfMax = 100.0 * diffValue / maxValue;
 				string red = "#B00000";
 				string green = "#008000";
-				if (percentOfMax >= 1.0)
+				if (percentOfMax >= 0.5)
 				{
 					return (formatInfo.autoColorizeMode == AutoColorizeMode.HighIsBad) ? red : green;
 				}
-				if (percentOfMax <= -1.0)
+				if (percentOfMax <= -0.5)
 				{
 					return (formatInfo.autoColorizeMode == AutoColorizeMode.HighIsBad) ? green : red;
 				}
