PerfReportTool 4.90 - Improve column formatting. Remove hacky min-width calculation and just rely on CSS no-wrap.

diff --git a/Engine/Source/Programs/CSVTools/PerfReportTool/PerfReportTool.cs b/Engine/Source/Programs/CSVTools/PerfReportTool/PerfReportTool.cs
--- a/Engine/Source/Programs/CSVTools/PerfReportTool/PerfReportTool.cs
+++ b/Engine/Source/Programs/CSVTools/PerfReportTool/PerfReportTool.cs
@@ -25,7 +25,7 @@ namespace PerfReportTool
 {
     class Version
     {
-        private static string VersionString = "4.89";
+        private static string VersionString = "4.90";
 
         public static string Get() { return VersionString; }
     };
diff --git a/Engine/Source/Programs/CSVTools/PerfReportTool/SummaryTable.cs b/Engine/Source/Programs/CSVTools/PerfReportTool/SummaryTable.cs
--- a/Engine/Source/Programs/CSVTools/PerfReportTool/SummaryTable.cs
+++ b/Engine/Source/Programs/CSVTools/PerfReportTool/SummaryTable.cs
@@ -1205,26 +1205,15 @@ namespace PerfSummaries
 
 			if (bScrollableTable)
 			{
-				int firstColMaxStringLength = 0;
-				if (columns.Count > 0)
-				{
-					for (int i = 0; i < columns[0].GetCount(); i++)
-					{
-						firstColMaxStringLength = Math.Max(firstColMaxStringLength, columns[0].GetStringValue(i).Length);
-					}
-				}
-
 				int headerMinWidth = bTranspose ? 50 : 75;
 				int cellFontSize = bTranspose ? 11 : 10;
 				int headerCellFontSize = cellFontSize - 1;
 
 				if (bAddMinMaxColumns)
 				{
-					headerMinWidth = 40;
+					headerMinWidth = 35;
 				}
 
-				int firstColWidth = (int)(firstColMaxStringLength * 6.5);
-
 				tableCss =
 					"table {table-layout: fixed;} \n" +
 					"table, th, td { border: 0px solid black; border-spacing: 0; border-collapse: separate; padding: " + cellPadding + "px; vertical-align: center; font-family: 'Verdana', Times, serif; font-size: "+ cellFontSize + "px;} \n" +
@@ -1253,8 +1242,7 @@ namespace PerfSummaries
 				// Top-left cell of the table is always on top, big font, thick border
 				tableCss += "tr:first-child th:first-child { z-index: 100;  border-right: 2px solid black; border-top: 2px solid black; font-size: 11px; top:0; left: 0px; } \n";
 
-				// Fix the first column width
-				tableCss += "th:first-child, td:first-child { border-left: 2px solid black; min-width: " + firstColWidth + ";} \n";
+				tableCss += "th:first-child, td:first-child { border-left: 2px solid black; white-space: nowrap;} \n";
 
 				if (bAddMinMaxColumns && isCollated)
 				{
