PerfReportTool - more formatting improvements

diff --git a/Engine/Source/Programs/CSVTools/PerfReportTool/SummaryTable.cs b/Engine/Source/Programs/CSVTools/PerfReportTool/SummaryTable.cs
--- a/Engine/Source/Programs/CSVTools/PerfReportTool/SummaryTable.cs
+++ b/Engine/Source/Programs/CSVTools/PerfReportTool/SummaryTable.cs
@@ -308,15 +308,21 @@ namespace PerfSummaries
 			List<string> newToolTips = new List<string>(toolTips.Count > 0 ? newCount : 0);
 			List<ColourThresholdList> newColourThresholds = new List<ColourThresholdList>(colourThresholds.Count > 0 ? newCount : 0);
 
+			bool bComputeDiff = isNumeric && name != "Count";
+
 			// Add diff rows to each of the arrays
 			for (int i = 0; i < doubleValues.Count; i++)
 			{
 				newDoubleValues.Add(doubleValues[i]);
 				if (i > 0)
 				{
-					double thisValue = FilterInvalidValue( doubleValues[i] );
-					double prevValue = FilterInvalidValue( doubleValues[i - 1] );
-					newDoubleValues.Add( thisValue - prevValue );
+					if (bComputeDiff)
+					{
+						double thisValue = FilterInvalidValue(doubleValues[i]);
+						double prevValue = FilterInvalidValue(doubleValues[i - 1]);
+						newDoubleValues.Add(thisValue - prevValue);
+					}
+					newDoubleValues.Add(0.0);
 				}
 			}
 			for (int i = 0; i < stringValues.Count; i++)
@@ -475,11 +481,25 @@ namespace PerfSummaries
 			const double absoluteIgnoreThreshold = 0.025;
 			if (hasDiffRows && IsDiffRow(index) && isNumeric && formatInfo.autoColorizeMode != AutoColorizeMode.Off && index < doubleValues.Count )
 			{
+				// For simplicity, just negate the diff value if lowIsBad
 				double diffValue = doubleValues[index];
-				if ( Math.Abs(diffValue) < absoluteIgnoreThreshold )
+				if (formatInfo.autoColorizeMode == AutoColorizeMode.LowIsBad)
+				{
+					diffValue *= -1.0;
+				}
+
+				// Diff absolute value is insignificant: output faded color
+				if (Math.Abs(diffValue) < absoluteIgnoreThreshold)
 				{
-					return "#A0A0A0";
+					// Very close to zero: just output grey
+					if (Math.Abs(diffValue) < 0.001f)
+					{
+						return "#A0A0A0";
+					}
+					// Slight red/green
+					return diffValue > 0.0 ? "#B8A0A0" : "#A0B8A0";
 				}
+
 				double prevValue = FilterInvalidValue(doubleValues[index - 2]);
 				double thisValue = FilterInvalidValue(doubleValues[index - 1]);
 
@@ -487,15 +507,17 @@ namespace PerfSummaries
 				double percentOfMax = 100.0 * diffValue / maxValue;
 				string red = "#B00000";
 				string green = "#008000";
+				// More than half a percent of max: output full colours
 				if (percentOfMax >= 0.5)
 				{
-					return (formatInfo.autoColorizeMode == AutoColorizeMode.HighIsBad) ? red : green;
+					return red;
 				}
 				if (percentOfMax <= -0.5)
 				{
-					return (formatInfo.autoColorizeMode == AutoColorizeMode.HighIsBad) ? green : red;
+					return green;
 				}
-				return "#808080";
+				// Output faded red/green
+				return diffValue > 0.0 ? "#B8A0A0" : "#A0B8A0";
 			}
 			return null;
 		}
@@ -1010,7 +1032,7 @@ namespace PerfSummaries
 			{
 				columns[i].AddDiffRows(i == 0);
 			}
-			rowCount = rowCount + (rowCount - 1);
+			rowCount += rowCount - 1;
 
 			// Just set rowWeightings to null for now. We shouldn't need it, since we should have already collated by this point
 			rowWeightings = null;
@@ -1206,8 +1228,8 @@ namespace PerfSummaries
 			if (bScrollableTable)
 			{
 				int headerMinWidth = bTranspose ? 50 : 75;
-				int cellFontSize = bTranspose ? 11 : 10;
-				int headerCellFontSize = cellFontSize - 1;
+				int cellFontSize = bTranspose ? 12 : 10;
+				int headerCellFontSize = bTranspose ? 10 : 9;
 
 				if (bAddMinMaxColumns)
 				{
@@ -1265,8 +1287,8 @@ namespace PerfSummaries
 
 			bool bOddRowsGray = !(!bAddMinMaxColumns || !isCollated);
 
-			string oddColor = bOddRowsGray ? "#e2e2e2" : "#ffffff";
-			string evenColor = bOddRowsGray ? "#ffffff" : "#e2e2e2";
+			string oddColor = bOddRowsGray ? "#eaeaea" : "#ffffff";
+			string evenColor = bOddRowsGray ? "#ffffff" : "#eaeaea";
 
 			tableCss += "tr:nth-child(odd) {background-color: "+oddColor+";} \n";
 			tableCss += "tr:nth-child(even) {background-color: "+evenColor+";} \n";
@@ -1345,7 +1367,7 @@ namespace PerfSummaries
 
 			// Work out which rows are major/minor section boundaries
 			Dictionary<int, int> rowSectionBoundaryLevel = new Dictionary<int, int>();
-			if (sectionBoundaries != null)
+			if (sectionBoundaries != null && !bTranspose)
 			{
 				foreach (SummarySectionBoundaryInfo sectionBoundaryInfo in sectionBoundaries)
 				{
@@ -1479,6 +1501,9 @@ namespace PerfSummaries
 			{
 				htmlTable = htmlTable.Transpose();
 				htmlTable.Set(0, 0, new HtmlTable.Cell(title, ""));
+
+				// Add a section boundary where the stats start
+				htmlTable.rows[htmlTable.numHeaderRows + firstStatColumnIndex - 1].attributes = " class='sectionStartLevel2'";
 			}
 
 			// Apply final formatting
