PerfReportTool 4.91 - replace autoColorize bool xml param with colorizeMode, which can be default, auto or off (new option). Add equivalent commandline param (-colorizeTable)
- Adding a diff row no longer disables colorization (unless colorizeMode is explicitly set to off)
- More formatting improvements for transposed tables

diff --git a/Engine/Source/Programs/CSVTools/PerfReportTool/PerfReportTool.cs b/Engine/Source/Programs/CSVTools/PerfReportTool/PerfReportTool.cs
--- a/Engine/Source/Programs/CSVTools/PerfReportTool/PerfReportTool.cs
+++ b/Engine/Source/Programs/CSVTools/PerfReportTool/PerfReportTool.cs
@@ -25,7 +25,7 @@ namespace PerfReportTool
 {
     class Version
     {
-        private static string VersionString = "4.90";
+        private static string VersionString = "4.91";
 
         public static string Get() { return VersionString; }
     };
@@ -147,7 +147,7 @@ namespace PerfReportTool
 			"       -noSummaryMinMax: don't make min/max columns for each stat in a condensed summary\n" +
 			"       -reverseTable [0|1]: Reverses the order of summary tables (set 0 to force off)\n" +
 			"       -scrollableTable [0|1]: makes the summary table scrollable, with frozen first rows and columns (set 0 to force off)\n" +
-			"       -autoColorizeTable [0|1]: enables summary table autocolorization. (set 0 to force off) \n" +
+			"       -colorizeTable [off|default|auto]: selects the table colorization mode\n" +
 			"       -maxSummaryTableStringLength <n>: strings longer than this will get truncated\n" +
 			"       -allowDuplicateCSVs : doesn't remove duplicate CSVs (Note: can cause summary table cache file locking issues)\n" +
 			"       -requireMetadata : ignores CSVs without metadata\n" +
@@ -735,11 +735,22 @@ namespace PerfReportTool
 				bScrollableTable = (bool)bScrollableTableOption;
 			}
 
-			bool bAutoColorizeTable = tableInfo.bAutoColorize;
-			bool? bAutoColorizeTableOption = GetOptionalBoolArg("autoColorizeTable");
-			if (bAutoColorizeTableOption != null)
+			TableColorizeMode colorizeMode = tableInfo.tableColorizeMode;
+			string colorizeArg = GetArg("colorizeTable","").ToLower();
+			if ( GetBoolArg("autoColorizeTable")) // Legacy support for the -autoColorizeTable arg
 			{
-				bAutoColorizeTable = (bool)bAutoColorizeTableOption;
+				colorizeMode = TableColorizeMode.Auto;
+			}
+			if ( colorizeArg != "" )
+			{
+				if (colorizeArg == "auto")
+				{
+					colorizeMode = TableColorizeMode.Auto;
+				}
+				else if (colorizeArg == "off")
+				{
+					colorizeMode = TableColorizeMode.Off;
+				}
 			}
 
 			bool addMinMaxColumns = !GetBoolArg("noSummaryMinMax") && !bTransposeCollatedSummaryTable;
@@ -775,8 +786,8 @@ namespace PerfReportTool
 					VersionString, 
 					bSpreadsheetFriendlyStrings, 
 					tableInfo.sectionBoundaries, 
-					bScrollableTable, 
-					bAutoColorizeTable, 
+					bScrollableTable,
+					colorizeMode, 
 					addMinMaxColumns, 
 					tableInfo.hideStatPrefix,
 					GetIntArg("maxSummaryTableStringLength", Int32.MaxValue), 
diff --git a/Engine/Source/Programs/CSVTools/PerfReportTool/SummaryTable.cs b/Engine/Source/Programs/CSVTools/PerfReportTool/SummaryTable.cs
--- a/Engine/Source/Programs/CSVTools/PerfReportTool/SummaryTable.cs
+++ b/Engine/Source/Programs/CSVTools/PerfReportTool/SummaryTable.cs
@@ -88,7 +88,20 @@ namespace PerfSummaries
 
 			bReverseSortRows = tableElement.GetSafeAttibute<bool>("reverseSortRows", false);
 			bScrollableFormatting = tableElement.GetSafeAttibute<bool>("scrollableFormatting", false);
-			bAutoColorize = tableElement.GetSafeAttibute<bool>("autoColorize", false);
+
+			string colorizeModeStr = tableElement.GetSafeAttibute<string>("colorizeMode", "").ToLower();
+			if (colorizeModeStr != "")
+			{
+				if (colorizeModeStr == "auto")
+				{
+					tableColorizeMode = TableColorizeMode.Auto;
+				}
+				else if (colorizeModeStr == "off")
+				{
+					tableColorizeMode = TableColorizeMode.Off;
+				}
+			}
+
 			statThreshold = tableElement.GetSafeAttibute<float>("statThreshold", 0.0f);
 			hideStatPrefix = tableElement.GetSafeAttibute<string>("hideStatPrefix");
 
@@ -152,7 +165,7 @@ namespace PerfSummaries
 		public List<SummarySectionBoundaryInfo> sectionBoundaries = new List<SummarySectionBoundaryInfo>();
 		public bool bReverseSortRows;
 		public bool bScrollableFormatting;
-		public bool bAutoColorize;
+		public TableColorizeMode tableColorizeMode = TableColorizeMode.Default;
 		public float statThreshold;
 		public string hideStatPrefix = null;
 		public string weightByColumn = null;
@@ -207,6 +220,13 @@ namespace PerfSummaries
 		static SummaryTableColumnFormatInfo defaultColumnInfo = new SummaryTableColumnFormatInfo();
 	};
 
+	enum TableColorizeMode
+	{
+		Off,
+		Default,
+		Auto,
+	};
+
 	enum AutoColorizeMode
 	{
 		Off,
@@ -366,7 +386,7 @@ namespace PerfSummaries
 			return ((rowIndex - 2) % 2) == 0;
 		}
 
-		public string GetDisplayName(string hideStatPrefix=null, bool bAddSeparatorSpaces = true)
+		public string GetDisplayName(string hideStatPrefix=null, bool bAddStatCategorySeparatorSpaces = true, bool bGreyOutStatCategories = false)
 		{
 			if (displayName != null)
 			{
@@ -376,21 +396,28 @@ namespace PerfSummaries
 			string statName = name;
 			if (hideStatPrefix != null)
 			{
-				string prefix = "";
-				string suffix = name;
-				if (name.StartsWith("Min ") || name.StartsWith("Max ") || name.StartsWith("Avg "))
+				string baseStatName = SummaryTable.GetBaseStatNameWithPrefixAndSuffix(statName, out string prefix, out string suffix);
+				if (baseStatName.ToLower().StartsWith(hideStatPrefix.ToLower() ) )
 				{
-					prefix = name.Substring(0, 4);
-					suffix = name.Substring(4);
-				}
-				if ( suffix.ToLower().StartsWith(hideStatPrefix.ToLower() ) )
-				{
-					statName = prefix + suffix.Substring(hideStatPrefix.Length);
+					statName = prefix + baseStatName.Substring(hideStatPrefix.Length) + suffix;
 				}
 			}
-			if (bAddSeparatorSpaces)
+
+			if (elementType == SummaryTableElement.Type.CsvStatAverage )
 			{
-				return TableUtil.FormatStatName(statName);
+				if (bGreyOutStatCategories)
+				{
+					string baseStatName = SummaryTable.GetBaseStatNameWithPrefixAndSuffix(statName, out string prefix, out _);
+					int idx = baseStatName.LastIndexOf("/");
+					if (idx >= 0)
+					{
+						statName = prefix + "<span class='greyText'>" + baseStatName.Substring(0, idx+1) + "</span><span class='blackText'>" + baseStatName.Substring(idx+1)+ "</span>";
+					}
+				}
+				if (bAddStatCategorySeparatorSpaces)
+				{
+					return statName.Replace("/", "/ ");
+				}
 			}
 			return statName;
 		}
@@ -454,7 +481,7 @@ namespace PerfSummaries
 		{
 			ColourThresholdList thresholds = null;
 			double value = GetValue(index);
-			if (value == double.MaxValue)
+			if (value == double.MaxValue || IsDiffRow(index))
 			{
 				return null;
 			}
@@ -479,7 +506,7 @@ namespace PerfSummaries
 		public string GetTextColor(int index)
 		{
 			const double absoluteIgnoreThreshold = 0.025;
-			if (hasDiffRows && IsDiffRow(index) && isNumeric && formatInfo.autoColorizeMode != AutoColorizeMode.Off && index < doubleValues.Count )
+			if (hasDiffRows && IsDiffRow(index) && isNumeric && index < doubleValues.Count )
 			{
 				// For simplicity, just negate the diff value if lowIsBad
 				double diffValue = doubleValues[index];
@@ -523,54 +550,61 @@ namespace PerfSummaries
 		}
 
 
-		public void ComputeAutomaticColourThresholds()
+		public void ComputeColorThresholds(TableColorizeMode tableColorizeMode)
 		{
-			AutoColorizeMode autoColorizeMode = formatInfo.autoColorizeMode;
-			if (autoColorizeMode == AutoColorizeMode.Off || !isNumeric)
+			if ( tableColorizeMode == TableColorizeMode.Default )
 			{
 				return;
 			}
-
-			if (hasDiffRows)
+			if (tableColorizeMode == TableColorizeMode.Off)
 			{
-				// If we have diff rows then color thresholds will be blank. We'll use font colors instead
+				// Set empty color thresholds. This clears existing thresholds from summaries
 				colourThresholds = new List<ColourThresholdList>();
+				return;
 			}
-			else // !hasDiffRows
+
+			AutoColorizeMode autoColorizeMode = formatInfo.autoColorizeMode;
+			if (autoColorizeMode == AutoColorizeMode.Off || !isNumeric)
 			{
-				// Set a single colour threshold list for the whole column
-				colourThresholds = new List<ColourThresholdList>();
-				double maxValue = -double.MaxValue;
-				double minValue = double.MaxValue;
-				double totalValue = 0.0f;
-				double validCount = 0.0f;
-				for (int i = 0; i < doubleValues.Count; i++)
-				{
-					double val = doubleValues[i];
-					if (val != double.MaxValue)
-					{
-						maxValue = Math.Max(val, maxValue);
-						minValue = Math.Min(val, minValue);
-						totalValue += val;
-						validCount += 1.0f;
-					}
+				return;
+			}
+
+			// Set a single colour threshold list for the whole column
+			colourThresholds = new List<ColourThresholdList>();
+			double maxValue = -double.MaxValue;
+			double minValue = double.MaxValue;
+			double totalValue = 0.0f;
+			double validCount = 0.0f;
+			for (int i = 0; i < doubleValues.Count; i++)
+			{
+				if (IsDiffRow(i))
+				{
+					continue;
 				}
-				if (minValue == maxValue)
+				double val = doubleValues[i];
+				if (val != double.MaxValue)
 				{
-					return;
+					maxValue = Math.Max(val, maxValue);
+					minValue = Math.Min(val, minValue);
+					totalValue += val;
+					validCount += 1.0f;
 				}
+			}
+			if (minValue == maxValue)
+			{
+				return;
+			}
 
-				Colour green = new Colour(0.4f, 0.82f, 0.45f);
-				Colour yellow = new Colour(1.0f, 1.0f, 0.5f);
-				Colour red = new Colour(1.0f, 0.4f, 0.4f);
+			Colour green = new Colour(0.4f, 0.82f, 0.45f);
+			Colour yellow = new Colour(1.0f, 1.0f, 0.5f);
+			Colour red = new Colour(1.0f, 0.4f, 0.4f);
 
-				double averageValue = totalValue / validCount; // TODO: Weighted average 
-				colourThresholdOverride = new ColourThresholdList();
-				colourThresholdOverride.Add(new ThresholdInfo(minValue, (autoColorizeMode == AutoColorizeMode.HighIsBad) ? green : red));
-				colourThresholdOverride.Add(new ThresholdInfo(averageValue, yellow));
-				colourThresholdOverride.Add(new ThresholdInfo(averageValue, yellow));
-				colourThresholdOverride.Add(new ThresholdInfo(maxValue, (autoColorizeMode == AutoColorizeMode.HighIsBad) ? red : green));
-			}
+			double averageValue = totalValue / validCount; // TODO: Weighted average 
+			colourThresholdOverride = new ColourThresholdList();
+			colourThresholdOverride.Add(new ThresholdInfo(minValue, (autoColorizeMode == AutoColorizeMode.HighIsBad) ? green : red));
+			colourThresholdOverride.Add(new ThresholdInfo(averageValue, yellow));
+			colourThresholdOverride.Add(new ThresholdInfo(averageValue, yellow));
+			colourThresholdOverride.Add(new ThresholdInfo(maxValue, (autoColorizeMode == AutoColorizeMode.HighIsBad) ? red : green));
 		}
 
 		public int GetCount()
@@ -666,12 +700,12 @@ namespace PerfSummaries
 					}
 					if (bIsDiffRow)
 					{
-						string strOut = val.ToString("0.000");
-						if (strOut == "0.000")
+						// Filter out close to zero results in diff columns
+						if (absVal < 0.000)
 						{
 							return "";
 						}
-						return prefix + strOut;
+						return prefix + val.ToString("0.00");
 					}
 					else
 					{
@@ -1046,9 +1080,7 @@ namespace PerfSummaries
 				if (statDisplaynameMapping != null && column.displayName == null)
 				{
 					string name = column.name;
-					string suffix = "";
-					string prefix = "";
-					string statName = GetStatNameWithPrefixAndSuffix(name, out prefix, out suffix);
+					string statName = GetBaseStatNameWithPrefixAndSuffix(name, out string prefix, out string suffix);
 					if (statDisplaynameMapping.ContainsKey(statName.ToLower()))
 					{
 						column.displayName = prefix + statDisplaynameMapping[statName.ToLower()] + suffix;
@@ -1056,15 +1088,12 @@ namespace PerfSummaries
 				}
 			}
 		}
-
-		string GetStatNameWithoutPrefixAndSuffix(string inName)
+		public static string GetBaseStatName(string inName)
 		{
-			string suffix = "";
-			string prefix = "";
-			return GetStatNameWithPrefixAndSuffix(inName, out prefix, out suffix);
+			return GetBaseStatNameWithPrefixAndSuffix(inName, out _, out _);
 		}
 
-		string GetStatNameWithPrefixAndSuffix(string inName, out string prefix, out string suffix)
+		public static string GetBaseStatNameWithPrefixAndSuffix(string inName, out string prefix, out string suffix)
 		{
 			suffix = "";
 			prefix = "";
@@ -1114,7 +1143,7 @@ namespace PerfSummaries
 			bool bSpreadsheetFriendlyStrings, 
 			List<SummarySectionBoundaryInfo> sectionBoundaries, 
 			bool bScrollableTable, 
-			bool bAutoColorizeTable,
+			TableColorizeMode tableColorizeMode,
 			bool bAddMinMaxColumns, 
 			string hideStatPrefix,
 			int maxColumnStringLength, 
@@ -1141,7 +1170,8 @@ namespace PerfSummaries
 			htmlFile.WriteLine("<html>");
 			htmlFile.WriteLine("<head><title>Perf Summary: "+ title + "</title>");
 
-			bool bAddStatNameSpacing = true;// !bTranspose;
+			bool bAddStatNameSpacing = !bTranspose;
+			bool bGreyOutStatPrefixes = bScrollableTable && bTranspose;
 
 			// Figure out the sticky column count
 			int stickyColumnCount = 0;
@@ -1168,14 +1198,11 @@ namespace PerfSummaries
 			}
 
 			// Automatically colourize the table if requested
-			if (bAutoColorizeTable)
+			if (tableColorizeMode != TableColorizeMode.Default)
 			{
 				foreach (SummaryTableColumn column in columns)
 				{
-					if (column.isNumeric)
-					{
-						column.ComputeAutomaticColourThresholds();
-					}
+					column.ComputeColorThresholds(tableColorizeMode);
 				}
 			}
 
@@ -1228,8 +1255,11 @@ namespace PerfSummaries
 			if (bScrollableTable)
 			{
 				int headerMinWidth = bTranspose ? 50 : 75;
+				int headerMaxWidth = bTranspose ? 165 : 220;
 				int cellFontSize = bTranspose ? 12 : 10;
 				int headerCellFontSize = bTranspose ? 10 : 9;
+				int firstColVerticalPadding = bTranspose ? 5 : 0;
+				string cellAlign = bTranspose ? "right" : "left";
 
 				if (bAddMinMaxColumns)
 				{
@@ -1238,16 +1268,23 @@ namespace PerfSummaries
 
 				tableCss =
 					"table {table-layout: fixed;} \n" +
-					"table, th, td { border: 0px solid black; border-spacing: 0; border-collapse: separate; padding: " + cellPadding + "px; vertical-align: center; font-family: 'Verdana', Times, serif; font-size: "+ cellFontSize + "px;} \n" +
+					"table, th, td { border: 0px solid black; border-spacing: 0; border-collapse: separate; padding: " + cellPadding + "px; vertical-align: center; font-family: 'Verdana', Times, serif; font-size: " + cellFontSize + "px;} \n" +
 					"td {" +
 					"  border-right: 1px solid black;" +
 					"  max-width: 450;" +
 					"} \n" +
+					"td:first-child {" +
+					"  padding-right:10px; padding-left:5px;" +
+					"  padding-top:" + firstColVerticalPadding + "px;padding-bottom:" + firstColVerticalPadding + "px;" +
+					"} \n" +
+					"td:not(:first-child) {" +
+					"  text-align: " + cellAlign + ";" +
+					"} \n" +
 					"tr:first-element { border-top: 2px; border-bottom: 2px } \n" +
 					"th {" +
-					"  width: auto;"+
-					"  max-width: 220px;" +
-					"  min-width: "+ headerMinWidth + "px;" +
+					"  width: auto;" +
+					"  max-width: " + headerMaxWidth + "px;" +
+					"  min-width: " + headerMinWidth + "px;" +
 					"  position: -webkit-sticky;" +
 					"  position: sticky;" +
 					"  border-right: 1px solid black;" +
@@ -1255,16 +1292,24 @@ namespace PerfSummaries
 					"  z-index: 5;" +
 					"  background-color: #ffffff;" +
 					"  top:0;" +
-					"  font-size: "+ headerCellFontSize + "px;" +
+					"  font-size: " + headerCellFontSize + "px;" +
 					"  word-wrap: break-word;" +
 					"  overflow: hidden;" +
 					"  height: 60;" +
+					"} \n" +
+					"span.greyText {" +
+					"  color: #808080;" +
+					"  display: inline-block;"+
+					"} \n"+
+					"span.blackText {" +
+					"  color: #000000;" +
+					"  display: inline-block;" +
 					"} \n";
 
 				// Top-left cell of the table is always on top, big font, thick border
 				tableCss += "tr:first-child th:first-child { z-index: 100;  border-right: 2px solid black; border-top: 2px solid black; font-size: 11px; top:0; left: 0px; } \n";
 
-				tableCss += "th:first-child, td:first-child { border-left: 2px solid black; white-space: nowrap;} \n";
+				tableCss += "th:first-child, td:first-child { border-left: 2px solid black; white-space: nowrap; max-width:800px;} \n";
 
 				if (bAddMinMaxColumns && isCollated)
 				{
@@ -1320,7 +1365,7 @@ namespace PerfSummaries
 				// Add the special columns (up to Count) to the lower header row
 				for (int i = 0; i < firstStatColumnIndex; i++)
 				{
-					headerRow.AddCell(columns[i].GetDisplayName(hideStatPrefix, bAddStatNameSpacing));
+					headerRow.AddCell(columns[i].GetDisplayName(hideStatPrefix, bAddStatNameSpacing, bGreyOutStatPrefixes));
 				}
 
 				if (bAddMinMaxColumns)
@@ -1338,7 +1383,7 @@ namespace PerfSummaries
 					// Add the stat columns
 					for (int i = firstStatColumnIndex; i < columns.Count; i++)
 					{
-						string statName = GetStatNameWithPrefixAndSuffix(columns[i].GetDisplayName(hideStatPrefix, bAddStatNameSpacing), out string prefix, out string suffix);
+						string statName = GetBaseStatNameWithPrefixAndSuffix(columns[i].GetDisplayName(hideStatPrefix, bAddStatNameSpacing, bGreyOutStatPrefixes), out string prefix, out string suffix);
 						if ((i - 1) % statColSpan == 0)
 						{
 							topHeaderRow.AddCell(statName + suffix, "colspan='" + statColSpan + "'");
@@ -1351,7 +1396,7 @@ namespace PerfSummaries
 					// Add the stat columns
 					for (int i = firstStatColumnIndex; i < columns.Count; i++)
 					{
-						string statName = GetStatNameWithPrefixAndSuffix(columns[i].GetDisplayName(hideStatPrefix, bAddStatNameSpacing), out _, out string suffix);
+						string statName = GetBaseStatNameWithPrefixAndSuffix(columns[i].GetDisplayName(hideStatPrefix, bAddStatNameSpacing, bGreyOutStatPrefixes), out _, out string suffix);
 						headerRow.AddCell(statName + suffix);
 					}
 				}
@@ -1360,7 +1405,7 @@ namespace PerfSummaries
 			{
 				foreach (SummaryTableColumn column in columns)
 				{
-					headerRow.AddCell(column.GetDisplayName(hideStatPrefix, bAddStatNameSpacing));
+					headerRow.AddCell(column.GetDisplayName(hideStatPrefix, bAddStatNameSpacing, bGreyOutStatPrefixes));
 				}
 			}
 			htmlTable.AddRow(headerRow);
@@ -1462,7 +1507,7 @@ namespace PerfSummaries
 						string toolTip = column.GetToolTipValue(rowIndex);
 						if (toolTip == "")
 						{
-							toolTip = column.GetDisplayName(null, bAddStatNameSpacing);
+							toolTip = column.GetDisplayName();
 						}
 						attributes.Add("title='" + toolTip + "'");
 					}
