PerfReportTool - Updated the Histogram summary table

- New option to suppress the averages table
- Reversed from under to over budget
- Removed the hitch table since the tool already has a seperate hitch summary

diff --git a/Engine/Source/Programs/CSVTools/PerfReportTool/PerfReportTool.cs b/Engine/Source/Programs/CSVTools/PerfReportTool/PerfReportTool.cs
--- a/Engine/Source/Programs/CSVTools/PerfReportTool/PerfReportTool.cs
+++ b/Engine/Source/Programs/CSVTools/PerfReportTool/PerfReportTool.cs
@@ -25,7 +25,7 @@ namespace PerfReportTool
 {
     class Version
     {
-        private static string VersionString = "4.91";
+        private static string VersionString = "4.92";
 
         public static string Get() { return VersionString; }
     };
diff --git a/Engine/Source/Programs/CSVTools/PerfReportTool/Summaries/HistogramSummary.cs b/Engine/Source/Programs/CSVTools/PerfReportTool/Summaries/HistogramSummary.cs
--- a/Engine/Source/Programs/CSVTools/PerfReportTool/Summaries/HistogramSummary.cs
+++ b/Engine/Source/Programs/CSVTools/PerfReportTool/Summaries/HistogramSummary.cs
@@ -24,13 +24,6 @@ namespace PerfSummaries
 				HistogramThresholds[i] = Convert.ToDouble(histogramStrings[i], System.Globalization.CultureInfo.InvariantCulture);
 			}
 
-			string[] hitchThresholds = element.Element("hitchThresholds").Value.Split(',');
-			HitchThresholds = new double[hitchThresholds.Length];
-			for (int i = 0; i < hitchThresholds.Length; i++)
-			{
-				HitchThresholds[i] = Convert.ToDouble(hitchThresholds[i], System.Globalization.CultureInfo.InvariantCulture);
-			}
-
 			foreach (XElement child in element.Elements())
 			{
 				if (child.Name == "budgetOverride")
@@ -39,6 +32,8 @@ namespace PerfSummaries
 					BudgetOverrideStatBudget = Convert.ToDouble(child.Attribute("budget").Value, System.Globalization.CultureInfo.InvariantCulture);
 				}
 			}
+
+			bSuppressAveragesTable = element.GetSafeAttibute<bool>("suppressAveragesTable", false);
 		}
 
 		public HistogramSummary() { }
@@ -52,26 +47,30 @@ namespace PerfSummaries
 			{
 				return;
 			}
-			// Write the averages
-			htmlFile.WriteLine("  <h2>Stat unit averages</h2>");
-			htmlFile.WriteLine("  <table border='0' style='width:400'>");
-			htmlFile.WriteLine("  <tr><td></td><th>ms</b></th>");
-			foreach (string stat in stats)
+
+			if (!bSuppressAveragesTable)
 			{
-				string StatToCheck = stat.Split('(')[0];
-				if (!csvStats.Stats.ContainsKey(StatToCheck.ToLower()))
+				// Write the averages
+				htmlFile.WriteLine("  <h2>Stat unit averages</h2>");
+				htmlFile.WriteLine("  <table border='0' style='width:400'>");
+				htmlFile.WriteLine("  <tr><td></td><th>ms</b></th>");
+				foreach (string stat in stats)
 				{
-					continue;
+					string StatToCheck = stat.Split('(')[0];
+					if (!csvStats.Stats.ContainsKey(StatToCheck.ToLower()))
+					{
+						continue;
+					}
+					float val = csvStats.Stats[StatToCheck.ToLower()].average;
+					string colour = ColourThresholdList.GetThresholdColour(val, ColourThresholds[0], ColourThresholds[1], ColourThresholds[2], ColourThresholds[3]);
+					htmlFile.WriteLine("  <tr><td><b>" + StatToCheck + "</b></td><td bgcolor=" + colour + ">" + val.ToString("0.00") + "</td></tr>");
 				}
-				float val = csvStats.Stats[StatToCheck.ToLower()].average;
-				string colour = ColourThresholdList.GetThresholdColour(val, ColourThresholds[0], ColourThresholds[1], ColourThresholds[2], ColourThresholds[3]);
-				htmlFile.WriteLine("  <tr><td><b>" + StatToCheck + "</b></td><td bgcolor=" + colour + ">" + val.ToString("0.00") + "</td></tr>");
+				htmlFile.WriteLine("  </table>");
 			}
-			htmlFile.WriteLine("  </table>");
 
-			// Hitches
+			// Histogram
 			double[] thresholds = HistogramThresholds;
-			htmlFile.WriteLine("  <h2>Frames in budget</h2>");
+			htmlFile.WriteLine("  <h2>Histogram</h2>");
 			htmlFile.WriteLine("  <table border='0' style='width:800'>");
 
 			htmlFile.WriteLine("  <tr><td></td>");
@@ -86,7 +85,7 @@ namespace PerfSummaries
 
 			foreach (float thresh in thresholds)
 			{
-				htmlFile.WriteLine("  <td><b><=" + thresh.ToString("0") + "ms</b></td>");
+				htmlFile.WriteLine("  <td><b>>" + thresh.ToString("0") + "ms</b></td>");
 			}
 			htmlFile.WriteLine("  </tr>");
 
@@ -102,8 +101,9 @@ namespace PerfSummaries
 				{
 					if (StatToCheck.ToLower() == BudgetOverrideStatName)
 					{
-						float pc = csvStats.GetStat(StatToCheck.ToLower()).GetRatioOfFramesInBudget((float)BudgetOverrideStatBudget) * 100.0f;
-						string colour = ColourThresholdList.GetThresholdColour(pc, 50.0f, 65.0f, 80.0f, 100.0f);
+						float RatioOfFramesOverBudget = 1.0f - csvStats.GetStat(StatToCheck.ToLower()).GetRatioOfFramesInBudget((float)BudgetOverrideStatBudget);
+						float pc = RatioOfFramesOverBudget * 100.0f;
+						string colour = ColourThresholdList.GetThresholdColour(pc, 50.0f, 35.0f, 20.0f, 0.0f);
 						htmlFile.WriteLine("  <td bgcolor=" + colour + ">" + pc.ToString("0.00") + "%</td>");
 					}
 					else
@@ -115,56 +115,22 @@ namespace PerfSummaries
 				foreach (float thresh in thresholds)
 				{
 					float threshold = (float)thresholds[thresholdIndex];
-					float pc = csvStats.GetStat(StatToCheck.ToLower()).GetRatioOfFramesInBudget(threshold) * 100.0f;
-					string colour = ColourThresholdList.GetThresholdColour(pc, 50.0f, 65.0f, 80.0f, 100.0f);
+					float RatioOfFramesOverBudget = 1.0f - csvStats.GetStat(StatToCheck.ToLower()).GetRatioOfFramesInBudget((float)threshold);
+					float pc = RatioOfFramesOverBudget * 100.0f;
+					string colour = ColourThresholdList.GetThresholdColour(pc, 50.0f, 35.0f, 20.0f, 0.0f);
 					htmlFile.WriteLine("  <td bgcolor=" + colour + ">" + pc.ToString("0.00") + "%</td>");
 					thresholdIndex++;
 				}
 				htmlFile.WriteLine("  </tr>");
 			}
 			htmlFile.WriteLine("  </table>");
-
-
-			// Hitches
-			htmlFile.WriteLine("  <h2>Hitches - Overall</h2>");
-			htmlFile.WriteLine("  <table border='0' style='width:800'>");
-			htmlFile.WriteLine("  <tr><td></td>");
-
-			foreach (float thresh in HitchThresholds)
-			{
-				htmlFile.WriteLine("  <td><b> >" + thresh.ToString("0") + "ms</b></td>");
-			}
-
-			htmlFile.WriteLine("  </tr>");
-
-			foreach (string unitStat in stats)
-			{
-				string StatToCheck = unitStat.Split('(')[0];
-				htmlFile.WriteLine("  <tr><td><b>" + unitStat + "</b></td>");
-				int thresholdIndex = 0;
-
-				foreach (float threshold in HitchThresholds)
-				{
-					float count = (float)csvStats.GetStat(StatToCheck.ToLower()).GetCountOfFramesOverBudget(threshold);
-					int numSamples = csvStats.GetStat(StatToCheck.ToLower()).GetNumSamples();
-					// if we have 20k frames in a typical flythrough then 20 frames would be red
-					float redThresholdFor50ms = (float)numSamples / 500.0f;
-					float redThreshold = (redThresholdFor50ms * 50.0f) / threshold; // Adjust the colour threshold based on the current threshold
-					string colour = ColourThresholdList.GetThresholdColour(count, redThreshold, redThreshold * 0.66, redThreshold * 0.33, 0.0f);
-					htmlFile.WriteLine("  <td bgcolor=" + colour + ">" + count.ToString("0") + "</td>");
-					thresholdIndex++;
-				}
-				htmlFile.WriteLine("  </tr>");
-			}
-			htmlFile.WriteLine("  </table>");
-
 		}
 
 		public double[] ColourThresholds;
 		public double[] HistogramThresholds;
-		public double[] HitchThresholds;
 		public string BudgetOverrideStatName;
 		public double BudgetOverrideStatBudget;
+		bool bSuppressAveragesTable;
 	};
 
 }
\ No newline at end of file
