[PerfReportTool] Added gamethread summary stat for percent of frames over fps target at {1.25x, 1.5x, 1.75x, 2x}. This can be used in summary tables to show a histogram distrubition.

diff --git a/Engine/Source/Programs/CSVTools/CsvStats/CsvStats.cs b/Engine/Source/Programs/CSVTools/CsvStats/CsvStats.cs
--- a/Engine/Source/Programs/CSVTools/CsvStats/CsvStats.cs
+++ b/Engine/Source/Programs/CSVTools/CsvStats/CsvStats.cs
@@ -316,7 +316,12 @@ namespace CSVStats
             return (float)(countInBudget) / (float)(samples.Count);
         }
 
-        public int GetCountOfFramesOverBudget(float budget, bool IgnoreFirstFrame = true, bool IgnoreLastFrame = true)
+		public float GetRatioOfFramesOverBudget(float budget)
+		{
+			return 1.0f - GetRatioOfFramesInBudget(budget);
+		}
+
+		public int GetCountOfFramesOverBudget(float budget, bool IgnoreFirstFrame = true, bool IgnoreLastFrame = true)
         {
             int countOverBudget = 0;
             for (int i = IgnoreFirstFrame ? 1 : 0; i < samples.Count - (IgnoreLastFrame ? 1 : 0); i++)
diff --git a/Engine/Source/Programs/CSVTools/PerfReportTool/PerfReportTool.cs b/Engine/Source/Programs/CSVTools/PerfReportTool/PerfReportTool.cs
--- a/Engine/Source/Programs/CSVTools/PerfReportTool/PerfReportTool.cs
+++ b/Engine/Source/Programs/CSVTools/PerfReportTool/PerfReportTool.cs
@@ -25,7 +25,7 @@ namespace PerfReportTool
 {
     class Version
     {
-        private static string VersionString = "4.92";
+        private static string VersionString = "4.93";
 
         public static string Get() { return VersionString; }
     };
diff --git a/Engine/Source/Programs/CSVTools/PerfReportTool/Summaries/FPSChartSummary.cs b/Engine/Source/Programs/CSVTools/PerfReportTool/Summaries/FPSChartSummary.cs
--- a/Engine/Source/Programs/CSVTools/PerfReportTool/Summaries/FPSChartSummary.cs
+++ b/Engine/Source/Programs/CSVTools/PerfReportTool/Summaries/FPSChartSummary.cs
@@ -259,6 +259,22 @@ namespace PerfSummaries
 					rowData.Add(SummaryTableElement.Type.SummaryTableMetric, columnName, column.Value, column.ColorThresholds);
 				}
 				rowData.Add(SummaryTableElement.Type.SummaryTableMetric, "TargetFPS", (double)fps);
+
+				// Write out gamethread histogram to the summary for 0.25x increments up to 2x
+				StatSamples gamethreadtimeStat = csvStats.GetStat("GameThreadtime");
+				if (gamethreadtimeStat != null)
+				{
+					float targetTimeInMS = 1000.0f / fps;
+					float redThreshold = 25.0f;
+					for (int i = 1; i <= 4; i++)
+					{
+						float targetTimeMultiplier = 1.0f + (0.25f * i);
+						float percentageOver = 100.0f * gamethreadtimeStat.GetRatioOfFramesOverBudget(targetTimeInMS * targetTimeMultiplier);
+						ColourThresholdList colorThresholdList = new ColourThresholdList(redThreshold, redThreshold * 0.66, redThreshold * 0.33, 0.0);
+
+						rowData.Add(SummaryTableElement.Type.SummaryTableMetric, "GT>" + targetTimeMultiplier + "x %", (double)percentageOver, colorThresholdList);
+					}
+				}
 			}
 
 			// Output HTML
