PerfReportTool 4.94 - Fix exception in HistogramSummary when SelectedStat doesn't exist. This was breaking perf reporting for devices without a GPUTime stat

diff --git a/Engine/Source/Programs/CSVTools/PerfReportTool/PerfReportTool.cs b/Engine/Source/Programs/CSVTools/PerfReportTool/PerfReportTool.cs
--- a/Engine/Source/Programs/CSVTools/PerfReportTool/PerfReportTool.cs
+++ b/Engine/Source/Programs/CSVTools/PerfReportTool/PerfReportTool.cs
@@ -25,7 +25,7 @@ namespace PerfReportTool
 {
     class Version
     {
-        private static string VersionString = "4.93";
+        private static string VersionString = "4.94";
 
         public static string Get() { return VersionString; }
     };
diff --git a/Engine/Source/Programs/CSVTools/PerfReportTool/Summaries/HistogramSummary.cs b/Engine/Source/Programs/CSVTools/PerfReportTool/Summaries/HistogramSummary.cs
--- a/Engine/Source/Programs/CSVTools/PerfReportTool/Summaries/HistogramSummary.cs
+++ b/Engine/Source/Programs/CSVTools/PerfReportTool/Summaries/HistogramSummary.cs
@@ -92,6 +92,10 @@ namespace PerfSummaries
 			foreach (string unitStat in stats)
 			{
 				string StatToCheck = unitStat.Split('(')[0];
+				if (!csvStats.Stats.ContainsKey(StatToCheck.ToLower()))
+				{
+					continue;
+				}
 
 				htmlFile.WriteLine("  <tr><td><b>" + StatToCheck + "</b></td>");
 				int thresholdIndex = 0;
