Make "BoundedStatValuesSummary: end event appeared before the start event" non-fatal.

diff --git a/Engine/Source/Programs/CSVTools/PerfReportTool/PerfReportTool.cs b/Engine/Source/Programs/CSVTools/PerfReportTool/PerfReportTool.cs
--- a/Engine/Source/Programs/CSVTools/PerfReportTool/PerfReportTool.cs
+++ b/Engine/Source/Programs/CSVTools/PerfReportTool/PerfReportTool.cs
@@ -25,7 +25,8 @@ namespace PerfReportTool
 {
     class Version
     {
-        private static string VersionString = "4.94";
+		// Format: Major.Minor.Bugfix
+        private static string VersionString = "4.94.1";
 
         public static string Get() { return VersionString; }
     };
diff --git a/Engine/Source/Programs/CSVTools/PerfReportTool/Summaries/BoundedStatValuesSummary.cs b/Engine/Source/Programs/CSVTools/PerfReportTool/Summaries/BoundedStatValuesSummary.cs
--- a/Engine/Source/Programs/CSVTools/PerfReportTool/Summaries/BoundedStatValuesSummary.cs
+++ b/Engine/Source/Programs/CSVTools/PerfReportTool/Summaries/BoundedStatValuesSummary.cs
@@ -133,7 +133,8 @@ namespace PerfSummaries
 			}
 			if (startFrame >= endFrame)
 			{
-				throw new Exception("BoundedStatValuesSummary: end event appeared before the start event");
+				Console.WriteLine("Warning: BoundedStatValuesSummary: end event "+ endEvent + " appeared before the start event "+beginEvent);
+				return;
 			}
 			endFrame = Math.Min(endFrame, csvStats.SampleCount - 1);
 			startFrame = Math.Max(startFrame, 0);
