PerfReportTool 4.94.2 - Fix diff row bug which caused an extra all-zero row to be created for every diff row added. This was breakng automatic colorization, causing all previous rows to be red

diff --git a/Engine/Source/Programs/CSVTools/PerfReportTool/PerfReportTool.cs b/Engine/Source/Programs/CSVTools/PerfReportTool/PerfReportTool.cs
--- a/Engine/Source/Programs/CSVTools/PerfReportTool/PerfReportTool.cs
+++ b/Engine/Source/Programs/CSVTools/PerfReportTool/PerfReportTool.cs
@@ -26,7 +26,7 @@ namespace PerfReportTool
     class Version
     {
 		// Format: Major.Minor.Bugfix
-        private static string VersionString = "4.94.1";
+        private static string VersionString = "4.94.2";
 
         public static string Get() { return VersionString; }
     };
diff --git a/Engine/Source/Programs/CSVTools/PerfReportTool/SummaryTable.cs b/Engine/Source/Programs/CSVTools/PerfReportTool/SummaryTable.cs
--- a/Engine/Source/Programs/CSVTools/PerfReportTool/SummaryTable.cs
+++ b/Engine/Source/Programs/CSVTools/PerfReportTool/SummaryTable.cs
@@ -342,7 +342,10 @@ namespace PerfSummaries
 						double prevValue = FilterInvalidValue(doubleValues[i - 1]);
 						newDoubleValues.Add(thisValue - prevValue);
 					}
-					newDoubleValues.Add(0.0);
+					else
+					{
+						newDoubleValues.Add(0.0);
+					}
 				}
 			}
 			for (int i = 0; i < stringValues.Count; i++)
