PerfReportTool : Minor QoL updates for histograms

- Show decimals on histogram thresholds
- Shift the gamethread summary table histogram down (75%-150%)
- Renamed the gamethread summary table histogram

diff --git a/Engine/Source/Programs/CSVTools/PerfReportTool/PerfReportTool.cs b/Engine/Source/Programs/CSVTools/PerfReportTool/PerfReportTool.cs
--- a/Engine/Source/Programs/CSVTools/PerfReportTool/PerfReportTool.cs
+++ b/Engine/Source/Programs/CSVTools/PerfReportTool/PerfReportTool.cs
@@ -26,7 +26,7 @@ namespace PerfReportTool
     class Version
     {
 		// Format: Major.Minor.Bugfix
-        private static string VersionString = "4.94.2";
+        private static string VersionString = "4.94.3";
 
         public static string Get() { return VersionString; }
     };
diff --git a/Engine/Source/Programs/CSVTools/PerfReportTool/Summaries/FPSChartSummary.cs b/Engine/Source/Programs/CSVTools/PerfReportTool/Summaries/FPSChartSummary.cs
--- a/Engine/Source/Programs/CSVTools/PerfReportTool/Summaries/FPSChartSummary.cs
+++ b/Engine/Source/Programs/CSVTools/PerfReportTool/Summaries/FPSChartSummary.cs
@@ -266,13 +266,13 @@ namespace PerfSummaries
 				{
 					float targetTimeInMS = 1000.0f / fps;
 					float redThreshold = 25.0f;
-					for (int i = 1; i <= 4; i++)
+					float[] targetTimeMultipliers = new float[] { 0.75f, 1.00f, 1.25f, 1.50f };
+					foreach(float targetTimeMultiplier in targetTimeMultipliers)
 					{
-						float targetTimeMultiplier = 1.0f + (0.25f * i);
 						float percentageOver = 100.0f * gamethreadtimeStat.GetRatioOfFramesOverBudget(targetTimeInMS * targetTimeMultiplier);
 						ColourThresholdList colorThresholdList = new ColourThresholdList(redThreshold, redThreshold * 0.66, redThreshold * 0.33, 0.0);
 
-						rowData.Add(SummaryTableElement.Type.SummaryTableMetric, "GT>" + targetTimeMultiplier + "x %", (double)percentageOver, colorThresholdList);
+						rowData.Add(SummaryTableElement.Type.SummaryTableMetric, "GameThread >" + targetTimeMultiplier + "x %", (double)percentageOver, colorThresholdList);
 					}
 				}
 			}
diff --git a/Engine/Source/Programs/CSVTools/PerfReportTool/Summaries/HistogramSummary.cs b/Engine/Source/Programs/CSVTools/PerfReportTool/Summaries/HistogramSummary.cs
--- a/Engine/Source/Programs/CSVTools/PerfReportTool/Summaries/HistogramSummary.cs
+++ b/Engine/Source/Programs/CSVTools/PerfReportTool/Summaries/HistogramSummary.cs
@@ -79,13 +79,13 @@ namespace PerfSummaries
 			bool HasBudgetOverrideStat = false;
 			if (BudgetOverrideStatName != null)
 			{
-				htmlFile.WriteLine("  <td><b><=" + BudgetOverrideStatBudget.ToString("0") + "ms</b></td>");
+				htmlFile.WriteLine("  <td><b><=" + BudgetOverrideStatBudget.ToString() + "ms</b></td>");
 				HasBudgetOverrideStat = true;
 			}
 
 			foreach (float thresh in thresholds)
 			{
-				htmlFile.WriteLine("  <td><b>>" + thresh.ToString("0") + "ms</b></td>");
+				htmlFile.WriteLine("  <td><b>>" + thresh.ToString() + "ms</b></td>");
 			}
 			htmlFile.WriteLine("  </tr>");
 
