Fixed inconsistencies with -colorizeTable and it not correctly overriding the colorMode set in the summary table xml.

+ 'default' now only exists as an argument and not an actual state as default means use what is defined in the summary table xml file.
+ Removed the 'Default' TableColorizeMode value and added a 'Budget' value which represents budget coloring.

diff --git a/Engine/Source/Programs/CSVTools/PerfReportTool/PerfReportTool.cs b/Engine/Source/Programs/CSVTools/PerfReportTool/PerfReportTool.cs
--- a/Engine/Source/Programs/CSVTools/PerfReportTool/PerfReportTool.cs
+++ b/Engine/Source/Programs/CSVTools/PerfReportTool/PerfReportTool.cs
@@ -26,7 +26,7 @@ namespace PerfReportTool
     class Version
     {
 		// Format: Major.Minor.Bugfix
-        private static string VersionString = "4.94.3";
+        private static string VersionString = "4.94.4";
 
         public static string Get() { return VersionString; }
     };
@@ -148,7 +148,7 @@ namespace PerfReportTool
 			"       -noSummaryMinMax: don't make min/max columns for each stat in a condensed summary\n" +
 			"       -reverseTable [0|1]: Reverses the order of summary tables (set 0 to force off)\n" +
 			"       -scrollableTable [0|1]: makes the summary table scrollable, with frozen first rows and columns (set 0 to force off)\n" +
-			"       -colorizeTable [off|default|auto]: selects the table colorization mode\n" +
+			"       -colorizeTable [off|budget|auto]: selects the table colorization mode. If omitted, uses the default in the summary xml table if set.\n" +
 			"       -maxSummaryTableStringLength <n>: strings longer than this will get truncated\n" +
 			"       -allowDuplicateCSVs : doesn't remove duplicate CSVs (Note: can cause summary table cache file locking issues)\n" +
 			"       -requireMetadata : ignores CSVs without metadata\n" +
@@ -736,6 +736,10 @@ namespace PerfReportTool
 				bScrollableTable = (bool)bScrollableTableOption;
 			}
 
+			// The colorize mode is initially set to whatever is in the summary table xml file.
+			// We then override that value if -colorizeTable is set to auto, off or budget.
+			// If -colorizeTable isn't specified then we use the default from the xml.
+			// If -colorizeTable isn't specified and it's not set in the summary table xml it uses the default set in the class initializer.
 			TableColorizeMode colorizeMode = tableInfo.tableColorizeMode;
 			string colorizeArg = GetArg("colorizeTable","").ToLower();
 			if ( GetBoolArg("autoColorizeTable")) // Legacy support for the -autoColorizeTable arg
@@ -752,6 +756,10 @@ namespace PerfReportTool
 				{
 					colorizeMode = TableColorizeMode.Off;
 				}
+				else if (colorizeArg == "budget")
+				{
+					colorizeMode = TableColorizeMode.Budget;
+				}
 			}
 
 			bool addMinMaxColumns = !GetBoolArg("noSummaryMinMax") && !bTransposeCollatedSummaryTable;
diff --git a/Engine/Source/Programs/CSVTools/PerfReportTool/SummaryTable.cs b/Engine/Source/Programs/CSVTools/PerfReportTool/SummaryTable.cs
--- a/Engine/Source/Programs/CSVTools/PerfReportTool/SummaryTable.cs
+++ b/Engine/Source/Programs/CSVTools/PerfReportTool/SummaryTable.cs
@@ -100,6 +100,10 @@ namespace PerfSummaries
 				{
 					tableColorizeMode = TableColorizeMode.Off;
 				}
+				else if (colorizeModeStr == "budget")
+				{
+					tableColorizeMode = TableColorizeMode.Budget;
+				}
 			}
 
 			statThreshold = tableElement.GetSafeAttibute<float>("statThreshold", 0.0f);
@@ -165,7 +169,7 @@ namespace PerfSummaries
 		public List<SummarySectionBoundaryInfo> sectionBoundaries = new List<SummarySectionBoundaryInfo>();
 		public bool bReverseSortRows;
 		public bool bScrollableFormatting;
-		public TableColorizeMode tableColorizeMode = TableColorizeMode.Default;
+		public TableColorizeMode tableColorizeMode = TableColorizeMode.Budget;
 		public float statThreshold;
 		public string hideStatPrefix = null;
 		public string weightByColumn = null;
@@ -222,9 +226,9 @@ namespace PerfSummaries
 
 	enum TableColorizeMode
 	{
-		Off,
-		Default,
-		Auto,
+		Off,		// No coloring.
+		Budget,     // Colorize based on defined budgets. No coloring for stats without defined budgets.
+		Auto,		// Auto calculate based on other values in the column.
 	};
 
 	enum AutoColorizeMode
@@ -555,7 +559,7 @@ namespace PerfSummaries
 
 		public void ComputeColorThresholds(TableColorizeMode tableColorizeMode)
 		{
-			if ( tableColorizeMode == TableColorizeMode.Default )
+			if ( tableColorizeMode == TableColorizeMode.Budget )
 			{
 				return;
 			}
@@ -1200,8 +1204,10 @@ namespace PerfSummaries
 				column.formatInfo = (columnFormatInfoList != null) ? columnFormatInfoList.GetFormatInfo(column.name) : SummaryTableColumnFormatInfoCollection.DefaultColumnInfo;
 			}
 
-			// Automatically colourize the table if requested
-			if (tableColorizeMode != TableColorizeMode.Default)
+			// Automatically colorize the table if requested.
+			// We run this even if colorize is off as we need to overwrite the values.
+			if (tableColorizeMode == TableColorizeMode.Auto ||
+				tableColorizeMode == TableColorizeMode.Off)
 			{
 				foreach (SummaryTableColumn column in columns)
 				{
