PerfReportTool - Add support for showing value buckets in columns.

This allows you to define a list of value thresholds for a column and the bucket names that correspond with each threshold. When present the corresponding bucket name will be displayed in the column instead of/with the value.

diff --git a/Engine/Source/Programs/CSVTools/PerfReportTool/PerfReportTool.cs b/Engine/Source/Programs/CSVTools/PerfReportTool/PerfReportTool.cs
--- a/Engine/Source/Programs/CSVTools/PerfReportTool/PerfReportTool.cs
+++ b/Engine/Source/Programs/CSVTools/PerfReportTool/PerfReportTool.cs
@@ -26,7 +26,7 @@ namespace PerfReportTool
     class Version
     {
 		// Format: Major.Minor.Bugfix
-        private static string VersionString = "4.94.4";
+        private static string VersionString = "4.95.0";
 
         public static string Get() { return VersionString; }
     };
diff --git a/Engine/Source/Programs/CSVTools/PerfReportTool/SummaryTable.cs b/Engine/Source/Programs/CSVTools/PerfReportTool/SummaryTable.cs
--- a/Engine/Source/Programs/CSVTools/PerfReportTool/SummaryTable.cs
+++ b/Engine/Source/Programs/CSVTools/PerfReportTool/SummaryTable.cs
@@ -267,12 +267,38 @@ namespace PerfSummaries
 			{
 				maxStringLengthCollated = maxStringLength;
 			}
+
+			includeValueWithBucketName = element.GetSafeAttibute<bool>("includeValueWithBucketName", true);
+			string bucketNamesString = element.GetSafeAttibute<string>("valueBucketNames");
+			if (bucketNamesString != null)
+			{
+				bucketNames = bucketNamesString.Split(',').ToList();
+			}
+			string bucketThresholdsString = element.GetSafeAttibute<string>("valueBucketThresholds");
+			if (bucketThresholdsString != null)
+			{
+				bucketThresholds = bucketThresholdsString.Split(',').Select(valStr =>
+				{
+					if (float.TryParse(valStr, out float value))
+					{
+						return value;
+					}
+					return 0.0f;
+				}).ToList();
+			}
 		}
+
 		public AutoColorizeMode autoColorizeMode = AutoColorizeMode.HighIsBad;
 		public string name;
 		public string numericFormat;
 		public int maxStringLength;
 		public int maxStringLengthCollated;
+		// If we should display the actual value in parenthesis next to the bucket name (if a bucket exists).
+		public bool includeValueWithBucketName = true;
+		// The name of each bucket. The name is indexed by the threshold.
+		public List<string> bucketNames = new List<string>();
+		// The value thresholds that correspond to each bucket. If a name doesn't exist for a threshold, the last bucket name is used.
+		public List<float> bucketThresholds = new List<float>();
 	};
 
 	class SummaryTableColumn
@@ -1538,6 +1564,31 @@ namespace PerfSummaries
 
 					string numericFormat = columnFormat.numericFormat;
 					string stringValue = column.GetStringValue(rowIndex, true, numericFormat);
+
+					// Check if we have any value buckets, if so lookup the bucket name for the value and display that with the value.
+					if (column.isNumeric && stringValue.Length > 0 && columnFormat.bucketNames.Count > 0 && columnFormat.bucketThresholds.Count > 0)
+					{
+						double value = column.GetValue(rowIndex);
+						int bucketIndex = 0;
+						for (bucketIndex = 0; bucketIndex < columnFormat.bucketThresholds.Count; ++bucketIndex)
+						{
+							if (value <= columnFormat.bucketThresholds[bucketIndex])
+							{
+								break;
+							}
+						}
+
+						bucketIndex = Math.Min(bucketIndex, columnFormat.bucketNames.Count-1);
+						if (columnFormat.includeValueWithBucketName)
+						{
+							stringValue = columnFormat.bucketNames[bucketIndex] + " (" + stringValue + ")";
+						}
+						else
+						{
+							stringValue = columnFormat.bucketNames[bucketIndex];
+						}
+					}
+
 					if (stringValue.Length > maxStringLength)
 					{
 						stringValue = TableUtil.SafeTruncateHtmlTableValue(stringValue, maxStringLength);
