PerfReportTool - Add support for colourThreshold overrides in summaryTableColumnFormatInfo.

This allows setting colourThresholds globally for an entire column, overriding other settings.

Also added support for specifiying the actual colours and if they should be lerped between or not.
Note: lerpColours isn't currently serialized to prc files.

diff --git a/Engine/Source/Programs/CSVTools/PerfReportTool/PerfReportTool.cs b/Engine/Source/Programs/CSVTools/PerfReportTool/PerfReportTool.cs
--- a/Engine/Source/Programs/CSVTools/PerfReportTool/PerfReportTool.cs
+++ b/Engine/Source/Programs/CSVTools/PerfReportTool/PerfReportTool.cs
@@ -26,7 +26,7 @@ namespace PerfReportTool
     class Version
     {
 		// Format: Major.Minor.Bugfix
-        private static string VersionString = "4.95.0";
+        private static string VersionString = "4.96.0";
 
         public static string Get() { return VersionString; }
     };
diff --git a/Engine/Source/Programs/CSVTools/PerfReportTool/Summary.cs b/Engine/Source/Programs/CSVTools/PerfReportTool/Summary.cs
--- a/Engine/Source/Programs/CSVTools/PerfReportTool/Summary.cs
+++ b/Engine/Source/Programs/CSVTools/PerfReportTool/Summary.cs
@@ -7,6 +7,7 @@ using System.Xml.Linq;
 using CSVStats;
 using PerfReportTool;
 using System.Reflection;
+using System.Drawing;
 
 namespace PerfSummaries
 {
@@ -195,23 +196,14 @@ namespace PerfSummaries
             return uniqueStats.ToArray();
         }
 
+		protected ColourThresholdList ReadColourThresholdListXML(XElement colourThresholdEl)
+		{
+			return ColourThresholdList.ReadColourThresholdListXML(colourThresholdEl);
+		}
+
 		protected double [] ReadColourThresholdsXML(XElement colourThresholdEl)
 		{
-			if (colourThresholdEl != null)
-			{
-				string[] colourStrings = colourThresholdEl.Value.Split(',');
-				if (colourStrings.Length != 4)
-				{
-					throw new Exception("Incorrect number of colourthreshold entries. Should be 4.");
-				}
-				double [] colourThresholds = new double[4];
-				for (int i = 0; i < colourStrings.Length; i++)
-				{
-					colourThresholds[i] = Convert.ToDouble(colourStrings[i], System.Globalization.CultureInfo.InvariantCulture);
-				}
-				return colourThresholds;
-			}
-			return null;
+			return ColourThresholdList.ReadColourThresholdsXML(colourThresholdEl);
 		}
 
         public string GetStatThresholdColour(string StatToUse, double value)
diff --git a/Engine/Source/Programs/CSVTools/PerfReportTool/SummaryTable.cs b/Engine/Source/Programs/CSVTools/PerfReportTool/SummaryTable.cs
--- a/Engine/Source/Programs/CSVTools/PerfReportTool/SummaryTable.cs
+++ b/Engine/Source/Programs/CSVTools/PerfReportTool/SummaryTable.cs
@@ -286,6 +286,8 @@ namespace PerfSummaries
 					return 0.0f;
 				}).ToList();
 			}
+
+			colourThresholdList = ColourThresholdList.ReadColourThresholdListXML(element.Element("colourThresholds"));
 		}
 
 		public AutoColorizeMode autoColorizeMode = AutoColorizeMode.HighIsBad;
@@ -299,6 +301,8 @@ namespace PerfSummaries
 		public List<string> bucketNames = new List<string>();
 		// The value thresholds that correspond to each bucket. If a name doesn't exist for a threshold, the last bucket name is used.
 		public List<float> bucketThresholds = new List<float>();
+		// Colour thresholds override for this column.
+		public ColourThresholdList colourThresholdList = null;
 	};
 
 	class SummaryTableColumn
@@ -518,7 +522,12 @@ namespace PerfSummaries
 			{
 				return null;
 			}
-			if (colourThresholdOverride != null)
+
+			if (formatInfo.colourThresholdList != null)
+			{
+				thresholds = formatInfo.colourThresholdList;
+			}
+			else if (colourThresholdOverride != null)
 			{
 				thresholds = colourThresholdOverride;
 			}
diff --git a/Engine/Source/Programs/CSVTools/PerfReportTool/SummaryTableData.cs b/Engine/Source/Programs/CSVTools/PerfReportTool/SummaryTableData.cs
--- a/Engine/Source/Programs/CSVTools/PerfReportTool/SummaryTableData.cs
+++ b/Engine/Source/Programs/CSVTools/PerfReportTool/SummaryTableData.cs
@@ -126,6 +126,7 @@ namespace PerfSummaries
 			{
 				int thresholdCount = reader.ReadInt32();
 				val.colorThresholdList = new ColourThresholdList();
+				// TODO: support serializing ColourThresholdList.LerpColours. This will require adding backwards compat to the serialization system.
 				for (int i = 0; i < thresholdCount; i++)
 				{
 					bool bHasColour = reader.ReadBoolean();
@@ -157,6 +158,7 @@ namespace PerfSummaries
 			if (colorThresholdList != null)
 			{
 				writer.Write((int)colorThresholdList.Count);
+				// TODO: support serializing ColourThresholdList.LerpColours. This will require adding backwards compat to the serialization system.
 				foreach (ThresholdInfo thresholdInfo in colorThresholdList.Thresholds)
 				{
 					writer.Write(thresholdInfo.colour != null);
diff --git a/Engine/Source/Programs/CSVTools/PerfReportTool/SummaryTypes.cs b/Engine/Source/Programs/CSVTools/PerfReportTool/SummaryTypes.cs
--- a/Engine/Source/Programs/CSVTools/PerfReportTool/SummaryTypes.cs
+++ b/Engine/Source/Programs/CSVTools/PerfReportTool/SummaryTypes.cs
@@ -2,6 +2,8 @@
 
 using System;
 using System.Collections.Generic;
+using System.Linq;
+using System.Xml.Linq;
 
 namespace PerfSummaries
 {
@@ -181,12 +183,19 @@ namespace PerfSummaries
 		public ColourThresholdList()
 		{ }
 
-		public ColourThresholdList(double redValue, double orangeValue, double yellowValue, double greenValue)
+		public ColourThresholdList(double redValue, double orangeValue, double yellowValue, double greenValue, bool lerpColours = true)
 		{
 			Add(new ThresholdInfo(greenValue, null));
 			Add(new ThresholdInfo(yellowValue, null));
 			Add(new ThresholdInfo(orangeValue, null));
 			Add(new ThresholdInfo(redValue, null));
+			LerpColours = lerpColours;
+		}
+
+		public ColourThresholdList(IEnumerable<ThresholdInfo> thresholds, bool lerpColours)
+		{
+			Thresholds = thresholds.ToList();
+			LerpColours = lerpColours;
 		}
 
 		public ColourThresholdList(Dictionary<string, dynamic> JsonDict)
@@ -208,6 +217,10 @@ namespace PerfSummaries
 						Thresholds[i].colour = new Colour(colorsList[i]);
 					}
 				}
+				if (JsonDict.ContainsKey("lerpColours"))
+				{
+					LerpColours = JsonDict["lerpColours"];
+				}
 			}
 		}
 
@@ -236,12 +249,58 @@ namespace PerfSummaries
 				{
 					Dict["colors"] = ThresholdColors;
 				}
+				Dict["lerpColours"] = LerpColours;
 			}
 			return Dict;
 		}
 
+		public static string GetThresholdColour(double value, List<ThresholdInfo> thresholds, bool lerpColours = true)
+		{
+			if (thresholds.Count == 0)
+			{
+				return Colour.White.ToHTMLString();
+			}
+
+			Colour col = null;
+			for (int i = 0; i < thresholds.Count; ++i)
+			{
+				ThresholdInfo threshold = thresholds[i];
+				if (threshold.colour == null)
+				{
+					throw new ArgumentException("All thresholds must have a valid colour set.", "thresholds");
+				}
+
+				if (value <= threshold.value)
+				{
+					if (i == 0)
+					{
+						col = threshold.colour;
+						break;
+					}
+
+					if (lerpColours)
+					{
+						ThresholdInfo prevThreshold = thresholds[i - 1];
+						double t = (value - prevThreshold.value) / (threshold.value - prevThreshold.value);
+						col = Colour.Lerp(prevThreshold.colour, threshold.colour, (float)t);
+					}
+					else
+					{
+						col = threshold.colour;
+					}
+					break;
+				}
+			}
+
+			if (col == null)
+			{
+				col = thresholds.Last().colour;
+			}
+			return col.ToHTMLString();
+		}
+
 		public static string GetThresholdColour(double value, double redValue, double orangeValue, double yellowValue, double greenValue,
-			Colour redOverride = null, Colour orangeOverride = null, Colour yellowOverride = null, Colour greenOverride = null)
+			Colour redOverride = null, Colour orangeOverride = null, Colour yellowOverride = null, Colour greenOverride = null, bool lerpColours = true)
 		{
 			Colour green = (greenOverride != null) ? greenOverride : new Colour(0.0f, 1.0f, 0.0f, 1.0f);
 			Colour orange = (orangeOverride != null) ? orangeOverride : new Colour(1.0f, 0.5f, 0.0f, 1.0f);
@@ -257,39 +316,19 @@ namespace PerfSummaries
 				value = -value;
 			}
 
-			Colour col = null;
-			if (value <= redValue)
-			{
-				col = red;
-			}
-			else if (value <= orangeValue)
-			{
-				double t = (value - redValue) / (orangeValue - redValue);
-				col = Colour.Lerp(red, orange, (float)t);
-			}
-			else if (value <= yellowValue)
-			{
-				double t = (value - orangeValue) / (yellowValue - orangeValue);
-				col = Colour.Lerp(orange, yellow, (float)t);
-			}
-			else if (value <= greenValue)
-			{
-				float t = (float)(value - yellowValue) / (float)(greenValue - yellowValue);
-				col = Colour.Lerp(yellow, green, t);
-			}
-			else
+			List<ThresholdInfo> thresholds = new List<ThresholdInfo>
 			{
-				col = green;
-			}
-			return col.ToHTMLString();
+				new ThresholdInfo(redValue, red),
+				new ThresholdInfo(orangeValue, orange),
+				new ThresholdInfo(yellowValue, yellow),
+				new ThresholdInfo(greenValue, green)
+			};
+			return GetThresholdColour(value, thresholds, lerpColours);
 		}
 
 		public void Add(ThresholdInfo info)
 		{
-			if (Thresholds.Count < 4)
-			{
-				Thresholds.Add(info);
-			}
+			Thresholds.Add(info);
 		}
 		public int Count
 		{
@@ -312,7 +351,11 @@ namespace PerfSummaries
 		{
 			if (Thresholds.Count == 4)
 			{
-				return GetThresholdColour(value, Thresholds[3].value, Thresholds[2].value, Thresholds[1].value, Thresholds[0].value, Thresholds[3].colour, Thresholds[2].colour, Thresholds[1].colour, Thresholds[0].colour);
+				return GetThresholdColour(value, Thresholds[3].value, Thresholds[2].value, Thresholds[1].value, Thresholds[0].value, Thresholds[3].colour, Thresholds[2].colour, Thresholds[1].colour, Thresholds[0].colour, LerpColours);
+			}
+			else if (Thresholds.Count > 0)
+			{
+				return GetThresholdColour(value, Thresholds, LerpColours);
 			}
 			return "'#ffffff'";
 		}
@@ -333,6 +376,75 @@ namespace PerfSummaries
 			}
 			return list.GetColourForValue(value);
 		}
+
+		public static ColourThresholdList ReadColourThresholdListXML(XElement colourThresholdEl)
+		{
+			if (colourThresholdEl != null)
+			{
+				ColourThresholdList thresholdList = new ColourThresholdList();
+
+				XAttribute lerpColoursElement = colourThresholdEl.Attribute("lerpColours");
+				if (lerpColoursElement != null)
+				{
+					thresholdList.LerpColours = bool.Parse(lerpColoursElement.Value);
+				}
+
+				XElement thresholdsElement = colourThresholdEl.Element("thresholds");
+				XElement coloursElement = colourThresholdEl.Element("colours");
+				if (thresholdsElement != null && coloursElement != null)
+				{
+					string[] thresholdStrings = thresholdsElement.Value.Split(',');
+					string[] colourStrings = coloursElement.Value.Split(',');
+					if (thresholdStrings.Length != colourStrings.Length)
+					{
+						throw new Exception("Number of colour thresholds and colours must match");
+					}
+
+					for (int i = 0; i < thresholdStrings.Length; ++i)
+					{
+						double value = double.Parse(thresholdStrings[i]);
+						Colour colour = new Colour(colourStrings[i]);
+						thresholdList.Add(new ThresholdInfo(value, colour));
+					}
+				}
+				else
+				{
+					// Revert to previous logic if sub elements don't exist.
+					double[] values = ReadColourThresholdsXML(colourThresholdEl);
+					if (values != null)
+					{
+						foreach (double value in values)
+						{
+							thresholdList.Add(new ThresholdInfo(value));
+						}
+					}
+				}
+				return thresholdList;
+			}
+			return null;
+		}
+
+		public static double[] ReadColourThresholdsXML(XElement colourThresholdEl)
+		{
+			if (colourThresholdEl != null)
+			{
+				string[] colourStrings = colourThresholdEl.Value.Split(',');
+				if (colourStrings.Length != 4)
+				{
+					throw new Exception("Incorrect number of colourthreshold entries. Should be 4.");
+				}
+				double[] colourThresholds = new double[4];
+				for (int i = 0; i < colourStrings.Length; i++)
+				{
+					colourThresholds[i] = Convert.ToDouble(colourStrings[i], System.Globalization.CultureInfo.InvariantCulture);
+				}
+				return colourThresholds;
+			}
+			return null;
+		}
+
 		public List<ThresholdInfo> Thresholds = new List<ThresholdInfo>();
+		// TODO: support serializing LerpColours. This will require adding backwards compat to the serialization system.
+		public bool LerpColours = true;
 	};
 }
\ No newline at end of file
