PerfReportTool - Update budget colour scheme to match auto as it's much easier on the eyes.

diff --git a/Engine/Source/Programs/CSVTools/PerfReportTool/PerfReportTool.cs b/Engine/Source/Programs/CSVTools/PerfReportTool/PerfReportTool.cs
--- a/Engine/Source/Programs/CSVTools/PerfReportTool/PerfReportTool.cs
+++ b/Engine/Source/Programs/CSVTools/PerfReportTool/PerfReportTool.cs
@@ -26,7 +26,7 @@ namespace PerfReportTool
     class Version
     {
 		// Format: Major.Minor.Bugfix
-        private static string VersionString = "4.96.0";
+        private static string VersionString = "4.97.0";
 
         public static string Get() { return VersionString; }
     };
diff --git a/Engine/Source/Programs/CSVTools/PerfReportTool/SummaryTable.cs b/Engine/Source/Programs/CSVTools/PerfReportTool/SummaryTable.cs
--- a/Engine/Source/Programs/CSVTools/PerfReportTool/SummaryTable.cs
+++ b/Engine/Source/Programs/CSVTools/PerfReportTool/SummaryTable.cs
@@ -637,9 +637,9 @@ namespace PerfSummaries
 				return;
 			}
 
-			Colour green = new Colour(0.4f, 0.82f, 0.45f);
-			Colour yellow = new Colour(1.0f, 1.0f, 0.5f);
-			Colour red = new Colour(1.0f, 0.4f, 0.4f);
+			Colour green = Colour.Green;
+			Colour yellow = Colour.Yellow;
+			Colour red = Colour.Red;
 
 			double averageValue = totalValue / validCount; // TODO: Weighted average 
 			colourThresholdOverride = new ColourThresholdList();
diff --git a/Engine/Source/Programs/CSVTools/PerfReportTool/SummaryTypes.cs b/Engine/Source/Programs/CSVTools/PerfReportTool/SummaryTypes.cs
--- a/Engine/Source/Programs/CSVTools/PerfReportTool/SummaryTypes.cs
+++ b/Engine/Source/Programs/CSVTools/PerfReportTool/SummaryTypes.cs
@@ -154,13 +154,13 @@ namespace PerfSummaries
 			int bI = (int)(b * 255.0f);
 			return "#" + rI.ToString("x2") + gI.ToString("x2") + bI.ToString("x2");
 		}
-		
+
 		public static Colour White = new Colour(1.0f, 1.0f, 1.0f, 1.0f);
 		public static Colour Black = new Colour(0, 0, 0, 1.0f);
-		public static Colour Orange = new Colour(1.0f, 0.5f, 0.0f, 1.0f);
-		public static Colour Yellow = new Colour(1.0f, 1.0f, 0.0f, 1.0f);
-		public static Colour Red = new Colour(1.0f, 0.0f, 0.0f, 1.0f);
-		public static Colour Green = new Colour(0.0f, 1.0f, 0.0f, 1.0f);
+		public static Colour Red = new Colour(1.0f, 0.4f, 0.4f);
+		public static Colour Orange = new Colour(1.0f, 0.7f, 0.4f, 1.0f);
+		public static Colour Yellow = new Colour(1.0f, 1.0f, 0.5f);
+		public static Colour Green = new Colour(0.4f, 0.82f, 0.45f);
 
 		public float r, g, b;
 		public float alpha;
@@ -302,10 +302,10 @@ namespace PerfSummaries
 		public static string GetThresholdColour(double value, double redValue, double orangeValue, double yellowValue, double greenValue,
 			Colour redOverride = null, Colour orangeOverride = null, Colour yellowOverride = null, Colour greenOverride = null, bool lerpColours = true)
 		{
-			Colour green = (greenOverride != null) ? greenOverride : new Colour(0.0f, 1.0f, 0.0f, 1.0f);
-			Colour orange = (orangeOverride != null) ? orangeOverride : new Colour(1.0f, 0.5f, 0.0f, 1.0f);
-			Colour yellow = (yellowOverride != null) ? yellowOverride : new Colour(1.0f, 1.0f, 0.0f, 1.0f);
-			Colour red = (redOverride != null) ? redOverride : new Colour(1.0f, 0.0f, 0.0f, 1.0f);
+			Colour green = (greenOverride != null) ? greenOverride : Colour.Green;
+			Colour orange = (orangeOverride != null) ? orangeOverride : Colour.Orange;
+			Colour yellow = (yellowOverride != null) ? yellowOverride : Colour.Yellow;
+			Colour red = (redOverride != null) ? redOverride : Colour.Red;
 
 			if (redValue > orangeValue)
 			{
