PerfReportTool - Add support for regression and other column filters.

+ Additional column filters can now be created by implementing ISummaryTableColumnFilter and adding it to the list in Program.MakeAdditionalColumnFilters.
+ Moved stat threshold checks into an additional filter.
+ Added a column filter to hide metadata columns.
+ Added a column filter to only show columns where the most recent build has regressed by more than a given threshold.
+ You can now use 'debugShowFilteredColumns' to have columns be grayed out instead of removed. Hovering over them will show the reason they were filtered out (which filter and the reason). This is for debugging purposes and sanity checks to make sure you're not missing important information, especially when iterating on new filters. Some additional settings were added to SummaryTableColumn to support this.
+ Filters can also mark individual cells as 'invalid' to indicate they are not taken into consideration for that filter. These are grayed out with a tooltip when debugShowFilteredColumns is enabled. Currently this does not work in collated views.
+ Added the ability to set tooltips for column headers.
+ Added the ability to add a color modifier for the entire column or individual cells.

diff --git a/Engine/Source/Programs/CSVTools/PerfReportTool/ColumnFilters/ISummaryTableColumnFilter.cs b/Engine/Source/Programs/CSVTools/PerfReportTool/ColumnFilters/ISummaryTableColumnFilter.cs
new file mode 100644
--- /dev/null
+++ b/Engine/Source/Programs/CSVTools/PerfReportTool/ColumnFilters/ISummaryTableColumnFilter.cs
@@ -0,0 +1,9 @@
+using PerfSummaries;
+
+namespace PerfReportTool
+{
+	internal interface ISummaryTableColumnFilter
+	{
+		bool ShouldFilter(SummaryTableColumn column, SummaryTable table);
+	}
+}
diff --git a/Engine/Source/Programs/CSVTools/PerfReportTool/ColumnFilters/MetadataColumnFilter.cs b/Engine/Source/Programs/CSVTools/PerfReportTool/ColumnFilters/MetadataColumnFilter.cs
new file mode 100644
--- /dev/null
+++ b/Engine/Source/Programs/CSVTools/PerfReportTool/ColumnFilters/MetadataColumnFilter.cs
@@ -0,0 +1,30 @@
+using PerfSummaries;
+using System.Collections.Generic;
+using System.Linq;
+
+namespace PerfReportTool
+{
+	// Filters out metadata columns.
+	internal class MetadataColumnFilter : ISummaryTableColumnFilter
+	{
+		// Metadata columns to not filter out.
+		private HashSet<string> ColumnsToIgnore;
+
+		public MetadataColumnFilter(List<string> columnsToNotFilter)
+		{
+			ColumnsToIgnore = columnsToNotFilter.ToHashSet();
+		}
+
+		public bool ShouldFilter(SummaryTableColumn column, SummaryTable table)
+		{
+			if ((column.elementType == SummaryTableElement.Type.CsvMetadata ||
+				 column.elementType == SummaryTableElement.Type.ToolMetadata) &&
+				 !ColumnsToIgnore.Contains(column.name))
+			{
+				column.DebugMarkAsFiltered(GetType().ToString(), "Column is metadata");
+				return true;
+			}
+			return false;
+		}
+	}
+}
diff --git a/Engine/Source/Programs/CSVTools/PerfReportTool/ColumnFilters/RegressionColumnFilter.cs b/Engine/Source/Programs/CSVTools/PerfReportTool/ColumnFilters/RegressionColumnFilter.cs
new file mode 100644
--- /dev/null
+++ b/Engine/Source/Programs/CSVTools/PerfReportTool/ColumnFilters/RegressionColumnFilter.cs
@@ -0,0 +1,177 @@
+using PerfSummaries;
+using System;
+using System.Collections.Generic;
+using System.Linq;
+
+namespace PerfReportTool
+{
+	// Filters out any columns where the stat hasn't regressed in the most row group.
+	internal class RegressionColumnFilter : ISummaryTableColumnFilter
+	{
+		// The stat name to join rows together that should be aggregated for sake of regression testing.
+		private string JoinRowsByStatName = null;
+		// Number of standard deviations from the mean the latest row group must be outside of to be considered a regression.
+		private double StdDevTreshold = 2.0;
+		// Number of standard deviations from the mean a value must be to be considered an outlier.
+		private double OutlierStdDevThreshold = 4.0;
+		// When the % of rows without a value is over this amount we just filter out the whole column due to lack of data.
+		private const double MissingRowThresholdPercent = 0.6;
+
+		public RegressionColumnFilter(string joinRowsByName, float stdDevThreshold, float outlierStdDevThreshold)
+		{
+			JoinRowsByStatName = joinRowsByName;
+			StdDevTreshold = stdDevThreshold;
+			OutlierStdDevThreshold = outlierStdDevThreshold;
+		}
+
+		// Overview of how this works:
+		// 1. If too many values are missing from the column we filter it out.
+		// 2. We calculate the mean and standard deviation of all rows except those belonging to the first row group.
+		// 3. We remove any rows whose values are outside the (aggressive) outlier std dev threshold as these will skew our mean.
+		// 4. We re-calculate the mean and std deviation of the rows now that the outliers are gone.
+		// 5. We calculate the mean and std dev of the rows belonging to the first row group.
+		// 6. If the mean of the first row group - the mean of the rest of the rows (minus outliers) is outside our std dev threshold then there's a regression.
+		public bool ShouldFilter(SummaryTableColumn column, SummaryTable table)
+		{
+			if (!column.isNumeric ||
+				(column.elementType != SummaryTableElement.Type.CsvStatAverage &&
+				 column.elementType != SummaryTableElement.Type.SummaryTableMetric))
+			{
+				return false;
+			}
+
+			int rowCount = table.Count;
+			int firstRowGroupEndIndex = GetFirstRowGroupEndOffset(table);
+
+			List<Tuple<int, double>> allValuesExcludingFirstRowGroup = GetNonEmptyRows(column, firstRowGroupEndIndex, rowCount);
+
+			// If too many rows are missing values then just ignore the whole column.
+			double percentMissing = 1.0 - ((double)allValuesExcludingFirstRowGroup.Count / (double)rowCount);
+			if (percentMissing > MissingRowThresholdPercent)
+			{
+				column.DebugMarkAsFiltered(GetType().ToString(), $"Column is missing {percentMissing * 100.0}% of its values, which is more than the threshold {MissingRowThresholdPercent * 100.0}%.");
+				return true;
+			}
+
+			bool highIsBad = true;
+			if (column.formatInfo != null)
+			{
+				highIsBad = column.formatInfo.autoColorizeMode == AutoColorizeMode.HighIsBad;
+			}
+
+			// Calculate the mean and std dev for all the rows except the ones belonging to the first row group.
+			double mean, stdDev;
+			CalculateStdDevAndMeanForRange(allValuesExcludingFirstRowGroup, out mean, out stdDev);
+
+			// Use the more aggressive std dev threshold to remove any outliers as these will skew our mean to the point where it may hide new regressions.
+			var rowsWithoutOutliers = new List<Tuple<int, double>>();
+			foreach (Tuple<int, double> row in allValuesExcludingFirstRowGroup)
+			{
+				int index = row.Item1;
+				double value = row.Item2;
+				if (!IsValueOutsideThreshold(value, mean, stdDev, OutlierStdDevThreshold, highIsBad))
+				{
+					rowsWithoutOutliers.Add(row);
+				}
+				else
+				{
+					column.DebugMarkRowInvalid(index, "RegressionFilter identified this value as an outlier.");
+				}
+			}
+
+			// Re-calculate with the outliers removed
+			CalculateStdDevAndMeanForRange(rowsWithoutOutliers, out mean, out stdDev);
+
+			// Grab the rows for the first group.
+			List<Tuple<int, double>> firstRowGroupValues = GetNonEmptyRows(column, 0, firstRowGroupEndIndex);
+			if (firstRowGroupValues.Count == 0)
+			{
+				column.DebugMarkAsFiltered(GetType().ToString(), "Missing values for the first row group.");
+				return true;
+			}
+
+			// Calculate the mean and std dev of the first row group.
+			double firstRowGroupMean, firstRowGroupStdDev;
+			CalculateStdDevAndMeanForRange(firstRowGroupValues, out firstRowGroupMean, out firstRowGroupStdDev);
+
+			// Check if the first row group is within the std dev threshold of the rest of the historical data excluding outliers.
+			if (IsValueOutsideThreshold(firstRowGroupMean, mean, stdDev, StdDevTreshold, highIsBad))
+			{
+				return false;
+			}
+
+			double modifier = highIsBad ? 1.0 : -1.0;
+			column.DebugMarkAsFiltered(GetType().ToString(),
+				$"First row group mean ({firstRowGroupMean:0.000}) is within the std dev threshold ({stdDev:0.000} * {StdDevTreshold:0.0} = {stdDev * StdDevTreshold:0.000}) of the total mean ({mean:0.000}). Difference = {(firstRowGroupMean - mean)* modifier:0.000}.");
+			return true;
+		}
+
+		private bool IsValueOutsideThreshold(double value, double mean, double stdDev, double numStdDevations, bool highIsBad)
+		{
+			double modifier = highIsBad ? 1.0 : -1.0;
+			double difference = (value - mean) * modifier;
+			double threshold = stdDev * numStdDevations;
+			return difference > threshold;
+		}
+
+		private void CalculateStdDevAndMeanForRange(List<Tuple<int, double>> rows, out double mean, out double stdDev)
+		{
+			var values = rows.Select(tuple => tuple.Item2).ToList();
+			CalculateStdDevAndMeanForRange(values, out mean, out stdDev);
+		}
+
+		private void CalculateStdDevAndMeanForRange(List<double> values, out double mean, out double stdDev)
+		{
+			int count = values.Count;
+			double total = values.Sum();
+			double avg = total / (double)count;
+
+			List<double> distFromMean = values.Select(value => Math.Pow(value - avg, 2)).ToList();
+			double totalDistFromMean = distFromMean.Sum();
+			stdDev = Math.Sqrt(totalDistFromMean / count);
+			mean = avg;
+		}
+
+		private List<Tuple<int, double>> GetNonEmptyRows(SummaryTableColumn column, int start, int end)
+		{
+			var values = new List<Tuple<int, double>>();
+			for (int i = start; i < end; ++i)
+			{
+				double value = column.GetValue(i);
+				if (value < double.MaxValue)
+				{
+					values.Add(Tuple.Create(i, value));
+				}
+			}
+			return values;
+		}
+
+		private int GetFirstRowGroupEndOffset(SummaryTable table)
+		{
+			if (JoinRowsByStatName == null)
+			{
+				// If no name is provided we treat each row distinctly.
+				return 1;
+			}
+
+			SummaryTableColumn groupByColumn = table.GetColumnByName(JoinRowsByStatName);
+			if (groupByColumn == null)
+			{
+				Console.WriteLine($"[Error] -regressionJoinRowsByName doesn't refer to a real column name: {JoinRowsByStatName}. Will not be applying any grouping to rows.");
+				return 1;
+			}
+
+			int firstGroupEndIndex = 0;
+			string firstRowValue = groupByColumn.GetStringValue(0);
+			for (int i = 1; i < table.Count; ++i)
+			{
+				if (groupByColumn.GetStringValue(i) != firstRowValue)
+				{
+					break;
+				}
+				++firstGroupEndIndex;
+			}
+			return firstGroupEndIndex;
+		}
+	}
+}
diff --git a/Engine/Source/Programs/CSVTools/PerfReportTool/ColumnFilters/StatThresholdColumnFilter.cs b/Engine/Source/Programs/CSVTools/PerfReportTool/ColumnFilters/StatThresholdColumnFilter.cs
new file mode 100644
--- /dev/null
+++ b/Engine/Source/Programs/CSVTools/PerfReportTool/ColumnFilters/StatThresholdColumnFilter.cs
@@ -0,0 +1,39 @@
+using PerfSummaries;
+using System.Collections.Generic;
+using System.Linq;
+
+namespace PerfReportTool
+{
+	// Filter out csv stat or metric columns below the specified threshold
+	internal class StatThresholdColumnFilter : ISummaryTableColumnFilter
+	{
+		private float StatThreshold;
+		public StatThresholdColumnFilter(float statThreshold)
+		{
+			StatThreshold = statThreshold;
+		}
+
+		public bool ShouldFilter(SummaryTableColumn column, SummaryTable table)
+		{
+			if (StatThreshold <= 0.0f)
+			{
+				return false;
+			}
+
+			if (!column.isNumeric ||
+				(column.elementType != SummaryTableElement.Type.CsvStatAverage &&
+				 column.elementType != SummaryTableElement.Type.SummaryTableMetric))
+			{
+				return false;
+			}
+
+			if (column.AreAllValuesOverThreshold((double)StatThreshold))
+			{
+				return false;
+			}
+
+			column.DebugMarkAsFiltered(GetType().ToString(), $"Not all column values are over thershold {StatThreshold}");
+			return true;
+		}
+	}
+}
diff --git a/Engine/Source/Programs/CSVTools/PerfReportTool/PerfReportTool.cs b/Engine/Source/Programs/CSVTools/PerfReportTool/PerfReportTool.cs
--- a/Engine/Source/Programs/CSVTools/PerfReportTool/PerfReportTool.cs
+++ b/Engine/Source/Programs/CSVTools/PerfReportTool/PerfReportTool.cs
@@ -26,7 +26,7 @@ namespace PerfReportTool
     class Version
     {
 		// Format: Major.Minor.Bugfix
-        private static string VersionString = "4.97.0";
+        private static string VersionString = "4.98.0";
 
         public static string Get() { return VersionString; }
     };
@@ -165,6 +165,14 @@ namespace PerfReportTool
 			"       -transposeTable : write the summary tables transposed\n"+
 			"       -transposeCollatedTable : write the collated summary table transposed (disables min/max columns)\n" +
             "       -addDiffRows : adds diff rows after the first two rows\n" + 
+			"\n" +
+			"Optional Column Filters\n" +
+			"		-debugShowFilteredColumns : grays out filtered columns instead of removing them. Column tooltip will show filtered reason.\n" +
+			"		-hideMetadataColumns : filters out metadata columns from the table (excluding those used in row sort).\n" +
+			"		-onlyShowRegressedColumns : only shows columns where the most recent row group (see -regressionJoinRowsByName) is outside the given standard deviation threshold from the mean of the previous rows.\n" +
+			"			-regressionJoinRowsByName <statName> (optional) : the csv stat name to join rows by for aggregation. If not provided each row is treated distinctly and no aggregation occurs.\n" +
+			"			-regressionStdDevThreshold <n> (default = 2) : how many standard deviations from the mean the most recent row group must be to be considered a regression.\n" +
+			"			-regressionOutlierStdDevThreshold <n> (default = 4) : how many standard deviations from the mean a value must be to be considered an outlier. Used in the initial pass to remove outliers to better identify real regressions.\n" +
 			"\n" +
 			"Json serialization:\n" +
 			"       -summaryTableToJson <filename> : json filename to write summary table row data to\n" +
@@ -769,8 +777,14 @@ namespace PerfReportTool
 				filenameWithoutExtension = Path.Combine(outputDir, filenameWithoutExtension);
 			}
 
-			float statThreshold = GetFloatArg("summaryTableStatThreshold", tableInfo.statThreshold);
-			SummaryTable filteredTable = table.SortAndFilter(tableInfo.columnFilterList, tableInfo.rowSortList, bReverseTable, weightByColumnName, statThreshold);
+			IEnumerable<ISummaryTableColumnFilter> additionalColumnFilters = MakeAdditionalColumnFilters(tableInfo);
+			bool showFilteredColumns = GetBoolArg("debugShowFilteredColumns");
+
+			// Set format info for the columns as some of the info is needed for the filters.
+			// TODO: would be better if we could determine HighIsBad without the format info and store it directly in the column.
+			table.SetColumnFormatInfo(reportXML.columnFormatInfoList);
+
+			SummaryTable filteredTable = table.SortAndFilter(tableInfo.columnFilterList, tableInfo.rowSortList, bReverseTable, weightByColumnName, showFilteredColumns, additionalColumnFilters);
 			if (bCollated)
 			{
 				filteredTable = filteredTable.CollateSortedTable(tableInfo.rowSortList, addMinMaxColumns);
@@ -790,6 +804,9 @@ namespace PerfReportTool
 					filteredTable.AddDiffRows();
 				}
 
+				// Run again to add format info for any new columns that were added (eg. count).
+				filteredTable.SetColumnFormatInfo(reportXML.columnFormatInfoList);
+
 				filteredTable.WriteToHTML(
 					filenameWithoutExtension + ".html", 
 					VersionString, 
@@ -800,14 +817,35 @@ namespace PerfReportTool
 					addMinMaxColumns, 
 					tableInfo.hideStatPrefix,
 					GetIntArg("maxSummaryTableStringLength", Int32.MaxValue), 
-					reportXML.columnFormatInfoList, 
 					weightByColumnName, 
 					summaryTitle,
-					bCollated ? bTransposeCollatedSummaryTable : bTransposeFullSummaryTable
+					bCollated ? bTransposeCollatedSummaryTable : bTransposeFullSummaryTable,
+					showFilteredColumns
 				);
 			}
 		}
 
+		IEnumerable<ISummaryTableColumnFilter> MakeAdditionalColumnFilters(SummaryTableInfo tableInfo)
+		{
+			List<ISummaryTableColumnFilter> additionalColumnFilters = new List<ISummaryTableColumnFilter>();
+			additionalColumnFilters.Add(new StatThresholdColumnFilter(GetFloatArg("summaryTableStatThreshold", tableInfo.statThreshold)));
+
+			if (GetBoolArg("hideMetadataColumns"))
+			{
+				additionalColumnFilters.Add(new MetadataColumnFilter(tableInfo.rowSortList));
+			}
+
+			if (GetBoolArg("onlyShowRegressedColumns"))
+			{
+				string joinByStatName = GetArg("regressionJoinRowsByName");
+				float stdDevThreshold = GetFloatArg("regressionStdDevThreshold", 2.0f);
+				float outlierStdDevThreshold = GetFloatArg("regressionOutlierStdDevThreshold", 4.0f);
+				additionalColumnFilters.Add(new RegressionColumnFilter(joinByStatName, stdDevThreshold, outlierStdDevThreshold));
+			}
+
+			return additionalColumnFilters;
+		}
+
 		string ReplaceFileExtension(string path, string newExtension)
 		{
 			// Special case for .bin.csv
diff --git a/Engine/Source/Programs/CSVTools/PerfReportTool/PerfreportTool.csproj b/Engine/Source/Programs/CSVTools/PerfReportTool/PerfreportTool.csproj
--- a/Engine/Source/Programs/CSVTools/PerfReportTool/PerfreportTool.csproj
+++ b/Engine/Source/Programs/CSVTools/PerfReportTool/PerfreportTool.csproj
@@ -56,6 +56,10 @@
     <Reference Include="System.Xml" />
   </ItemGroup>
   <ItemGroup>
+    <Compile Include="ColumnFilters\ISummaryTableColumnFilter.cs" />
+    <Compile Include="ColumnFilters\MetadataColumnFilter.cs" />
+    <Compile Include="ColumnFilters\RegressionColumnFilter.cs" />
+    <Compile Include="ColumnFilters\StatThresholdColumnFilter.cs" />
     <Compile Include="CsvFileCache..cs" />
     <Compile Include="PerfReportTool.cs" />
     <Compile Include="Summaries\BucketSummary.cs" />
diff --git a/Engine/Source/Programs/CSVTools/PerfReportTool/SummaryTable.cs b/Engine/Source/Programs/CSVTools/PerfReportTool/SummaryTable.cs
--- a/Engine/Source/Programs/CSVTools/PerfReportTool/SummaryTable.cs
+++ b/Engine/Source/Programs/CSVTools/PerfReportTool/SummaryTable.cs
@@ -6,6 +6,8 @@ using System.Linq;
 using System.Xml.Linq;
 using PerfReportTool;
 using System.IO;
+using System.Data.Common;
+using System.Runtime.InteropServices;
 
 namespace PerfSummaries
 {
@@ -312,30 +314,48 @@ namespace PerfSummaries
 		public string displayName;
 		public bool isRowWeightColumn = false;
 		public bool hasDiffRows = false;
+		// Column header tooltip. Displayed when hovering over the header.
+		public string tooltip = null;
+		// Multiplied with the colour of each cell to modify the colour (including header).
+		public Colour columnColourModifier = null;
 		List<double> doubleValues = new List<double>();
 		List<string> stringValues = new List<string>();
 		List<string> toolTips = new List<string>();
+		// Per cell colour modifiers.
+		Dictionary<int, Colour> colourModifiers = new Dictionary<int, Colour>();
 		public SummaryTableElement.Type elementType;
 		public SummaryTableColumnFormatInfo formatInfo = null;
 
 		List<ColourThresholdList> colourThresholds = new List<ColourThresholdList>();
 		ColourThresholdList colourThresholdOverride = null;
 
-		public SummaryTableColumn(string inName, bool inIsNumeric, string inDisplayName, bool inIsRowWeightColumn, SummaryTableElement.Type inElementType)
+		public SummaryTableColumn(
+			string inName,
+			bool inIsNumeric,
+			string inDisplayName,
+			bool inIsRowWeightColumn,
+			SummaryTableElement.Type inElementType,
+			SummaryTableColumnFormatInfo inFormatInfo = null,
+			string inTooltip = null,
+			Colour inColumnColourModifier = null)
 		{
 			name = inName;
 			isNumeric = inIsNumeric;
 			displayName = inDisplayName;
 			isRowWeightColumn = inIsRowWeightColumn;
 			elementType = inElementType;
+			formatInfo = inFormatInfo;
+			tooltip = inTooltip;
+			columnColourModifier = inColumnColourModifier;
 		}
 		public SummaryTableColumn Clone()
 		{
-			SummaryTableColumn newColumn = new SummaryTableColumn(name, isNumeric, displayName, isRowWeightColumn, elementType);
+			SummaryTableColumn newColumn = new SummaryTableColumn(name, isNumeric, displayName, isRowWeightColumn, elementType, formatInfo, tooltip, columnColourModifier);
 			newColumn.doubleValues.AddRange(doubleValues);
 			newColumn.stringValues.AddRange(stringValues);
 			newColumn.colourThresholds.AddRange(colourThresholds);
 			newColumn.toolTips.AddRange(toolTips);
+			newColumn.colourModifiers = colourModifiers.ToDictionary(entry => entry.Key, entry => entry.Value); // Deep copy
 			newColumn.hasDiffRows = hasDiffRows;
 			return newColumn;
 		}
@@ -542,7 +562,25 @@ namespace PerfSummaries
 					return null;
 				}
 			}
-			return thresholds.GetColourForValue(value);
+
+			Colour modifier = null;
+			if (columnColourModifier != null)
+			{
+				// Column modifier takes precedence over the cell modifier
+				modifier = columnColourModifier;
+			}
+			else if (colourModifiers.ContainsKey(index))
+			{
+				modifier = colourModifiers[index];
+			}
+
+			string colourString = thresholds.GetColourForValue(value);
+			if (modifier != null)
+			{
+				var colour = new Colour(colourString.Replace("'", ""));
+				colourString = (colour * modifier).ToHTMLString();
+			}
+			return colourString;
 		}
 
 		public string GetTextColor(int index)
@@ -649,6 +687,20 @@ namespace PerfSummaries
 			colourThresholdOverride.Add(new ThresholdInfo(maxValue, (autoColorizeMode == AutoColorizeMode.HighIsBad) ? red : green));
 		}
 
+		public List<string> GetHeaderAttributes()
+		{
+			var attributes = new List<string>();
+			if (columnColourModifier != null)
+			{
+				attributes.Add($"style='background-color:{Colour.White * columnColourModifier};'");
+			}
+			if (tooltip != null)
+			{
+				attributes.Add($"title='{tooltip}'");
+			}
+			return attributes;
+		}
+
 		public int GetCount()
 		{
 			return Math.Max(doubleValues.Count, stringValues.Count);
@@ -796,6 +848,17 @@ namespace PerfSummaries
 			return toolTips[index];
 		}
 
+		public void DebugMarkRowInvalid(int index, string reason)
+		{
+			colourModifiers.Add(index, new Colour(0.5f, 0.5f, 0.5f, 0.5f));
+			SetToolTipValue(index, $"Cell invalid: {reason}");
+		}
+
+		public void DebugMarkAsFiltered(string filterName, string reason)
+		{
+			tooltip = $"Filtered out by: {filterName}: {reason}";
+			columnColourModifier = new Colour(0.5f, 0.5f, 0.5f, 0.5f);
+		}
 	};
 
 
@@ -806,6 +869,14 @@ namespace PerfSummaries
 		{
 		}
 
+		public void SetColumnFormatInfo(SummaryTableColumnFormatInfoCollection collection)
+		{
+			foreach (SummaryTableColumn column in columns)
+			{
+				column.formatInfo = collection != null ? collection.GetFormatInfo(column.name) : SummaryTableColumnFormatInfoCollection.DefaultColumnInfo;
+			}
+		}
+
 		public SummaryTable CollateSortedTable(List<string> collateByList, bool addMinMaxColumns)
 		{
 			int numSubColumns = addMinMaxColumns ? 3 : 1;
@@ -832,7 +903,7 @@ namespace PerfSummaries
 			{
 				if (collateByColumns.Contains(srcColumn))
 				{
-					newColumns.Add(new SummaryTableColumn(srcColumn.name, false, srcColumn.displayName, false, srcColumn.elementType));
+					newColumns.Add(new SummaryTableColumn(srcColumn.name, false, srcColumn.displayName, false, srcColumn.elementType, srcColumn.formatInfo));
 					finalSortByList.Add(srcColumn.name.ToLower());
 					// Early out if we've found all the columns
 					if (finalSortByList.Count == collateByColumns.Count)
@@ -853,11 +924,11 @@ namespace PerfSummaries
 				if (column.isNumeric && !collateByColumns.Contains(column))
 				{
 					srcToDestBaseColumnIndex.Add(newColumns.Count);
-					newColumns.Add(new SummaryTableColumn("Avg " + column.name, true, null, false, column.elementType));
+					newColumns.Add(new SummaryTableColumn("Avg " + column.name, true, null, false, column.elementType, column.formatInfo, column.tooltip, column.columnColourModifier));
 					if (addMinMaxColumns)
 					{
-						newColumns.Add(new SummaryTableColumn("Min " + column.name, true, null, false, column.elementType));
-						newColumns.Add(new SummaryTableColumn("Max " + column.name, true, null, false, column.elementType));
+						newColumns.Add(new SummaryTableColumn("Min " + column.name, true, null, false, column.elementType, column.formatInfo, column.tooltip, column.columnColourModifier));
+						newColumns.Add(new SummaryTableColumn("Max " + column.name, true, null, false, column.elementType, column.formatInfo, column.tooltip, column.columnColourModifier));
 					}
 				}
 				else
@@ -994,7 +1065,7 @@ namespace PerfSummaries
 			return SortAndFilter(customFilter.Split(',').ToList(), customRowSort.Split(',').ToList(), bReverseSort, weightByColumnName);
 		}
 
-		public SummaryTable SortAndFilter(List<string> columnFilterList, List<string> rowSortList, bool bReverseSort, string weightByColumnName, float statThreshold = 0.0f)
+		public SummaryTable SortAndFilter(List<string> columnFilterList, List<string> rowSortList, bool bReverseSort, string weightByColumnName, bool showFilteredColumns = false, IEnumerable<ISummaryTableColumnFilter> additionalFilters = null)
 		{
 			SummaryTable newTable = SortRows(rowSortList, bReverseSort);
 
@@ -1078,19 +1149,30 @@ namespace PerfSummaries
 				}
 			}
 
-
-			// Filter out csv stat or metric columns below the specified threshold
-			if (statThreshold > 0.0f)
+			// Run the additional filters on the columns
+			if (additionalFilters != null)
 			{
-				List<SummaryTableColumn> oldColumnList = newColumnList;
-				newColumnList = new List<SummaryTableColumn>();
-				foreach (SummaryTableColumn column in oldColumnList)
+				var filteredColumns = new HashSet<string>();
+				foreach (ISummaryTableColumnFilter filter in additionalFilters)
 				{
-					if (!column.isNumeric || 
-						( column.elementType != SummaryTableElement.Type.CsvStatAverage && column.elementType != SummaryTableElement.Type.SummaryTableMetric ) || 
-						column.AreAllValuesOverThreshold((double)statThreshold))
+					if (showFilteredColumns)
+					{
+						// Run the filter so we can populate which columns would have been filtered, but don't actually remove them from the list.
+						newColumnList.ForEach(column =>
+						{
+							// If it's already filtered then skip it so we don't overwrite the filter reason.
+							if (!filteredColumns.Contains(column.name))
+							{
+								if (filter.ShouldFilter(column, this))
+								{
+									filteredColumns.Add(column.name);
+								}
+							}
+						});
+					}
+					else
 					{
-						newColumnList.Add(column);
+						newColumnList = newColumnList.Where(column => !filter.ShouldFilter(column, this)).ToList();
 					}
 				}
 			}
@@ -1189,10 +1271,10 @@ namespace PerfSummaries
 			bool bAddMinMaxColumns, 
 			string hideStatPrefix,
 			int maxColumnStringLength, 
-			SummaryTableColumnFormatInfoCollection columnFormatInfoList, 
 			string weightByColumnName, 
 			string title, 
-			bool bTranspose )
+			bool bTranspose,
+			bool showFilteredColumnDebug = false)
 		{
 			System.IO.StreamWriter htmlFile = new System.IO.StreamWriter(htmlFilename, false);
 			int statColSpan = hasMinMaxColumns ? 3 : 1;
@@ -1233,12 +1315,6 @@ namespace PerfSummaries
 				}
 			}
 
-			// Get format info for the columns
-			foreach (SummaryTableColumn column in columns)
-			{
-				column.formatInfo = (columnFormatInfoList != null) ? columnFormatInfoList.GetFormatInfo(column.name) : SummaryTableColumnFormatInfoCollection.DefaultColumnInfo;
-			}
-
 			// Automatically colorize the table if requested.
 			// We run this even if colorize is off as we need to overwrite the values.
 			if (tableColorizeMode == TableColorizeMode.Auto ||
@@ -1427,10 +1503,12 @@ namespace PerfSummaries
 					// Add the stat columns
 					for (int i = firstStatColumnIndex; i < columns.Count; i++)
 					{
-						string statName = GetBaseStatNameWithPrefixAndSuffix(columns[i].GetDisplayName(hideStatPrefix, bAddStatNameSpacing, bGreyOutStatPrefixes), out string prefix, out string suffix);
+						SummaryTableColumn column = columns[i];
+						string statName = GetBaseStatNameWithPrefixAndSuffix(column.GetDisplayName(hideStatPrefix, bAddStatNameSpacing, bGreyOutStatPrefixes), out string prefix, out string suffix);
 						if ((i - 1) % statColSpan == 0)
 						{
-							topHeaderRow.AddCell(statName + suffix, "colspan='" + statColSpan + "'");
+							string attributes = $"colspan='{statColSpan}'{String.Join(" ", column.GetHeaderAttributes())}";
+							topHeaderRow.AddCell(statName + suffix, attributes);
 						}
 						headerRow.AddCell(prefix.Trim());
 					}
@@ -1441,7 +1519,7 @@ namespace PerfSummaries
 					for (int i = firstStatColumnIndex; i < columns.Count; i++)
 					{
 						string statName = GetBaseStatNameWithPrefixAndSuffix(columns[i].GetDisplayName(hideStatPrefix, bAddStatNameSpacing, bGreyOutStatPrefixes), out _, out string suffix);
-						headerRow.AddCell(statName + suffix);
+						headerRow.AddCell(statName + suffix, String.Join(" ", columns[i].GetHeaderAttributes()));
 					}
 				}
 			}
@@ -1449,7 +1527,7 @@ namespace PerfSummaries
 			{
 				foreach (SummaryTableColumn column in columns)
 				{
-					headerRow.AddCell(column.GetDisplayName(hideStatPrefix, bAddStatNameSpacing, bGreyOutStatPrefixes));
+					headerRow.AddCell(column.GetDisplayName(hideStatPrefix, bAddStatNameSpacing, bGreyOutStatPrefixes), String.Join(" ", column.GetHeaderAttributes()));
 				}
 			}
 			htmlTable.AddRow(headerRow);
@@ -1685,7 +1763,7 @@ namespace PerfSummaries
 			List<SummaryTableColumn> newColumns = new List<SummaryTableColumn>();
 			foreach (SummaryTableColumn srcCol in columns)
 			{
-				SummaryTableColumn destCol = new SummaryTableColumn(srcCol.name, srcCol.isNumeric, null, false, srcCol.elementType);
+				SummaryTableColumn destCol = new SummaryTableColumn(srcCol.name, srcCol.isNumeric, null, false, srcCol.elementType, srcCol.formatInfo);
 				for (int i = 0; i < rowCount; i++)
 				{
 					int srcIndex = columnRemapping[i].Value;
@@ -1766,6 +1844,15 @@ namespace PerfSummaries
 			get { return rowCount; }
 		}
 
+		public SummaryTableColumn GetColumnByName(string name)
+		{
+			if (columnLookup.ContainsKey(name))
+			{
+				return columnLookup[name];
+			}
+			return null;
+		}
+
 		Dictionary<string, SummaryTableColumn> columnLookup = new Dictionary<string, SummaryTableColumn>();
 		List<SummaryTableColumn> columns = new List<SummaryTableColumn>();
 		List<double> rowWeightings = null;
diff --git a/Engine/Source/Programs/CSVTools/PerfReportTool/SummaryTypes.cs b/Engine/Source/Programs/CSVTools/PerfReportTool/SummaryTypes.cs
--- a/Engine/Source/Programs/CSVTools/PerfReportTool/SummaryTypes.cs
+++ b/Engine/Source/Programs/CSVTools/PerfReportTool/SummaryTypes.cs
@@ -142,6 +142,11 @@ namespace PerfSummaries
 			return Interpolated;
 		}
 
+		public static Colour operator *(Colour a, Colour b)
+		{
+			return new Colour(a.r * b.r, a.g * b.g, a.b * b.b, a.alpha * b.alpha);
+		}
+
 		public string ToHTMLString()
 		{
 			return "'" + ToString() + "'";
