PerfReportTool - Fix off by one error in regression filter when counting rows in first row group.

+ Also added handling for if all data is part of the same row group or if there are no rows (avoid divide by 0).

diff --git a/Engine/Source/Programs/CSVTools/PerfReportTool/ColumnFilters/RegressionColumnFilter.cs b/Engine/Source/Programs/CSVTools/PerfReportTool/ColumnFilters/RegressionColumnFilter.cs
--- a/Engine/Source/Programs/CSVTools/PerfReportTool/ColumnFilters/RegressionColumnFilter.cs
+++ b/Engine/Source/Programs/CSVTools/PerfReportTool/ColumnFilters/RegressionColumnFilter.cs
@@ -41,9 +41,20 @@ namespace PerfReportTool
 			}
 
 			int rowCount = table.Count;
-			int firstRowGroupEndIndex = GetFirstRowGroupEndOffset(table);
+			if (rowCount == 0)
+			{
+				column.DebugMarkAsFiltered(GetType().ToString(), "Column doesn't have any rows.");
+				return true;
+			}
+
+			int firstRowGroupCount = GetNumRowsInFirstRowGroup(table);
+			if (firstRowGroupCount == rowCount)
+			{
+				column.DebugMarkAsFiltered(GetType().ToString(), "All rows of data belong to the same group so there's nothing to compare against.");
+				return true;
+			}
 
-			List<Tuple<int, double>> allValuesExcludingFirstRowGroup = GetNonEmptyRows(column, firstRowGroupEndIndex, rowCount);
+			List<Tuple<int, double>> allValuesExcludingFirstRowGroup = GetNonEmptyRows(column, firstRowGroupCount, rowCount);
 
 			// If too many rows are missing values then just ignore the whole column.
 			double percentMissing = 1.0 - ((double)allValuesExcludingFirstRowGroup.Count / (double)rowCount);
@@ -83,7 +94,7 @@ namespace PerfReportTool
 			CalculateStdDevAndMeanForRange(rowsWithoutOutliers, out mean, out stdDev);
 
 			// Grab the rows for the first group.
-			List<Tuple<int, double>> firstRowGroupValues = GetNonEmptyRows(column, 0, firstRowGroupEndIndex);
+			List<Tuple<int, double>> firstRowGroupValues = GetNonEmptyRows(column, 0, firstRowGroupCount);
 			if (firstRowGroupValues.Count == 0)
 			{
 				column.DebugMarkAsFiltered(GetType().ToString(), "Missing values for the first row group.");
@@ -146,7 +157,7 @@ namespace PerfReportTool
 			return values;
 		}
 
-		private int GetFirstRowGroupEndOffset(SummaryTable table)
+		private int GetNumRowsInFirstRowGroup(SummaryTable table)
 		{
 			if (JoinRowsByStatName == null)
 			{
@@ -161,17 +172,15 @@ namespace PerfReportTool
 				return 1;
 			}
 
-			int firstGroupEndIndex = 0;
 			string firstRowValue = groupByColumn.GetStringValue(0);
-			for (int i = 1; i < table.Count; ++i)
+			for (int count = 1; count < table.Count; ++count)
 			{
-				if (groupByColumn.GetStringValue(i) != firstRowValue)
+				if (groupByColumn.GetStringValue(count) != firstRowValue)
 				{
-					break;
+					return count;
 				}
-				++firstGroupEndIndex;
 			}
-			return firstGroupEndIndex;
+			return table.Count;
 		}
 	}
 }
diff --git a/Engine/Source/Programs/CSVTools/PerfReportTool/PerfReportTool.cs b/Engine/Source/Programs/CSVTools/PerfReportTool/PerfReportTool.cs
--- a/Engine/Source/Programs/CSVTools/PerfReportTool/PerfReportTool.cs
+++ b/Engine/Source/Programs/CSVTools/PerfReportTool/PerfReportTool.cs
@@ -26,7 +26,7 @@ namespace PerfReportTool
     class Version
     {
 		// Format: Major.Minor.Bugfix
-        private static string VersionString = "4.98.0";
+        private static string VersionString = "4.98.1";
 
         public static string Get() { return VersionString; }
     };
