Added output of per-event mvp, hitches data so we can reference it in summary tables.

diff --git a/Engine/Source/Programs/CSVTools/PerfReportTool/Summaries/FPSChartSummary.cs b/Engine/Source/Programs/CSVTools/PerfReportTool/Summaries/FPSChartSummary.cs
--- a/Engine/Source/Programs/CSVTools/PerfReportTool/Summaries/FPSChartSummary.cs
+++ b/Engine/Source/Programs/CSVTools/PerfReportTool/Summaries/FPSChartSummary.cs
@@ -403,6 +403,24 @@ namespace PerfSummaries
 						column.UpdateColor();
 					}
 
+					// Write out data per capture range to summary table row data
+					if (rowData != null)
+                    {
+						if (!bIgnoreHitchTimePercent)
+						{
+							rowData.Add(SummaryTableElement.Type.SummaryTableMetric, CapRange.name + "_HitchTimePercent", captureFpsChartData.HitchTimePercent, GetStatColourThresholdList("HitchTimePercent"));
+						}
+
+						if (!bIgnoreMVP)
+                        {
+							rowData.Add(SummaryTableElement.Type.SummaryTableMetric, CapRange.name + "_MVP", captureFpsChartData.MVP, GetStatColourThresholdList(MvpStatName));
+						}
+
+						rowData.Add(SummaryTableElement.Type.SummaryTableMetric, CapRange.name + "_HPM", captureFpsChartData.HitchesPerMinute, GetStatColourThresholdList("Hitches/Min"));
+						rowData.Add(SummaryTableElement.Type.SummaryTableMetric, CapRange.name + "_HitchCount", captureFpsChartData.HitchCount);
+						rowData.Add(SummaryTableElement.Type.SummaryTableMetric, CapRange.name + "_TotalTimeSeconds", captureFpsChartData.TotalTimeSeconds);
+					}
+
 					// Output HTML
 					if (htmlFile != null)
 					{
