PerfReportTool 4.99.0
- Add support for a default report type. This will be considered if autodetection otherwise fails
- Add a requiredForDefaultFallback param to make specific metadata matches optional if we're falling back. This defaults to 1 (true)
- Add allowAutoDetect flag to reports so they can opt out of autodetection (these need to be specified explicitly)

diff --git a/Engine/Source/Programs/CSVTools/PerfReportTool/PerfReportTool.cs b/Engine/Source/Programs/CSVTools/PerfReportTool/PerfReportTool.cs
--- a/Engine/Source/Programs/CSVTools/PerfReportTool/PerfReportTool.cs
+++ b/Engine/Source/Programs/CSVTools/PerfReportTool/PerfReportTool.cs
@@ -26,7 +26,7 @@ namespace PerfReportTool
     class Version
     {
 		// Format: Major.Minor.Bugfix
-        private static string VersionString = "4.98.1";
+        private static string VersionString = "4.99.1";
 
         public static string Get() { return VersionString; }
     };
diff --git a/Engine/Source/Programs/CSVTools/PerfReportTool/ReportXml.cs b/Engine/Source/Programs/CSVTools/PerfReportTool/ReportXml.cs
--- a/Engine/Source/Programs/CSVTools/PerfReportTool/ReportXml.cs
+++ b/Engine/Source/Programs/CSVTools/PerfReportTool/ReportXml.cs
@@ -158,6 +158,7 @@ namespace PerfReportTool
 				}
 
 			}
+			defaultReportTypeName = reportTypesElement.GetSafeAttibute<string>("default");
 
 			Console.Out.WriteLine("GraphXML:  " + graphsXMLFilename+"\n");
 			XDocument reportGraphsDoc = XDocument.Load(graphsXMLFilename);
@@ -300,18 +301,40 @@ namespace PerfReportTool
 			ReportTypeInfo reportTypeInfo = null;
 			if (reportType == "")
 			{
+				XElement defaultReportTypeElement = null;
 				// Attempt to determine the report type automatically based on the stats
 				foreach (XElement element in reportTypesElement.Elements("reporttype"))
 				{
-					if (IsReportTypeXMLCompatibleWithStats(element, csvFile.dummyCsvStats))
+					bool bReportTypeSupportsAutodetect = element.GetSafeAttibute<bool>("allowAutoDetect", true);
+
+					if (bReportTypeSupportsAutodetect && IsReportTypeXMLCompatibleWithStats(element, csvFile.dummyCsvStats))
 					{
 						reportTypeInfo = new ReportTypeInfo(element, sharedSummaries, baseXmlDirectory);
 						break;
 					}
+
+					if (defaultReportTypeName != null && element.GetSafeAttibute<string>("name") == defaultReportTypeName )
+					{
+						defaultReportTypeElement = element;
+					}
 				}
-				if (reportTypeInfo == null)
+				// Attempt to fall back to a default ReportType if we didn't find one
+				if (reportTypeInfo == null && defaultReportTypeName != null)
 				{
-					throw new Exception("Compatible report type for CSV " + csvFile.filename + " could not be found in" + reportTypeXmlFilename);
+					if ( defaultReportTypeElement == null )
+					{
+						throw new Exception("Default report type " + defaultReportTypeName + " was not found in " + reportTypeXmlFilename);
+					}
+					if (!IsReportTypeXMLCompatibleWithStats(defaultReportTypeElement, csvFile.dummyCsvStats, true))
+					{
+						throw new Exception("Default report type " + defaultReportTypeName + " was not compatible with CSV " + csvFile.filename);
+					}
+					Console.Out.WriteLine("Falling back to default report type: " + defaultReportTypeName);
+					reportTypeInfo = new ReportTypeInfo(defaultReportTypeElement, sharedSummaries, baseXmlDirectory);
+				}
+				else if (reportTypeInfo == null)
+				{
+					throw new Exception("Compatible report type for CSV " + csvFile.filename + " could not be found in " + reportTypeXmlFilename);
 				}
 			}
 			else
@@ -364,7 +387,7 @@ namespace PerfReportTool
 			return reportTypeInfo;
 		}
 
-		bool IsReportTypeXMLCompatibleWithStats(XElement reportTypeElement, CsvStats csvStats)
+		bool IsReportTypeXMLCompatibleWithStats(XElement reportTypeElement, CsvStats csvStats, bool bIsDefaultFallback=false)
 		{
 			XAttribute nameAt = reportTypeElement.Attribute("name");
 			if (nameAt == null)
@@ -410,6 +433,12 @@ namespace PerfReportTool
 					return false;
 				}
 
+				// Some metadata may be safe to skip for the default fallback case
+				if (bIsDefaultFallback && requiredMetadataEl.GetSafeAttibute("ignoreForDefaultFallback", false))
+				{
+					continue;
+				}
+
 				bool ignoreIfKeyNotFound = requiredMetadataEl.GetSafeAttibute("ignoreIfKeyNotFound", true);
 				bool stopIfKeyFound = requiredMetadataEl.GetSafeAttibute("stopIfKeyFound", false);
 
@@ -474,6 +503,7 @@ namespace PerfReportTool
 		XElement rootElement;
 		XElement graphGroupsElement;
 		XElement summaryTablesElement;
+		string defaultReportTypeName;
 		Dictionary<string, XElement> sharedSummaries;
 		Dictionary<string, GraphSettings> graphs;
 		Dictionary<string, string> statDisplayNameMapping;
