PerfReportTool - fix help formatting so it fits in a standard 120-character wide command prompt and make some descriptions more concise.

diff --git a/Engine/Source/Programs/CSVTools/PerfReportTool/PerfReportTool.cs b/Engine/Source/Programs/CSVTools/PerfReportTool/PerfReportTool.cs
--- a/Engine/Source/Programs/CSVTools/PerfReportTool/PerfReportTool.cs
+++ b/Engine/Source/Programs/CSVTools/PerfReportTool/PerfReportTool.cs
@@ -26,7 +26,7 @@ namespace PerfReportTool
     class Version
     {
 		// Format: Major.Minor.Bugfix
-        private static string VersionString = "4.99.1";
+        private static string VersionString = "4.99.2";
 
         public static string Get() { return VersionString; }
     };
@@ -83,118 +83,124 @@ namespace PerfReportTool
 			"PerfReportTool v" + Version.Get() + "\n" +
 			"\n" +
 			"Format: \n" +
-			"       -csv <filename> or -csvdir <directory path> or -summaryTableCacheIn <directory path> or\n" +
-			"       -csvList <comma separated> or -prcList <comma separated>\n" +
-			"       -o <dir name>: output directory (will be created if necessary)\n" +
+			"  -csv <filename> or -csvdir <directory path> or -summaryTableCacheIn <directory path> or\n" +
+			"  -csvList <comma separated> or -prcList <comma separated>\n" +
+			"  -o <dir name>: output directory (will be created if necessary)\n" +
 			"\n" +
 			"Optional Args:\n" +
-			"       -reportType <e.g. flythrough, playthrough, playthroughmemory>\n" +
-			"       -reportTypeCompatCheck : do a compatibility if when specifying a report type (rather than forcing)\n" +
-			"       -graphXML <xmlfilename>\n" +
-			"       -reportXML <xmlfilename>\n" +
-			"       -reportxmlbasedir <folder>\n" +
-			"       -title <name> - title for detailed reports\n" +
-			"       -summaryTitle <name> - title for summary tables\n" +
-			"       -maxy <value> - forces all graphs to use this value\n" +
-			"       -writeSummaryCsv : if specified, a csv file containing summary information will be generated.\n" +
-			"          Not available in bulk mode.\n" +
-			"       -noWatermarks : don't embed the commandline or version in reports\n" +
-			"       -cleanCsvOut <filename> : write a standard format CSV after event stripping with metadata stripped out.\n" +
-			"          Not available in bulk mode.\n" +
-			"       -noSmooth : disable smoothing on all graphs\n" +
-			"       -listSummaryTables: lists available summary tables from the current report XML\n" +
+			"  -reportType <e.g. flythrough, playthrough, playthroughmemory>\n" +
+			"  -reportTypeCompatCheck : do a compatibility if when specifying a report type (rather than forcing)\n" +
+			"  -graphXML <xmlfilename>\n" +
+			"  -reportXML <xmlfilename>\n" +
+			"  -reportxmlbasedir <folder>\n" +
+			"  -title <name> - title for detailed reports\n" +
+			"  -summaryTitle <name> - title for summary tables\n" +
+			"  -maxy <value> - forces all graphs to use this value\n" +
+			"  -writeSummaryCsv : if specified, a csv file containing summary information will be generated.\n" +
+			"     Not available in bulk mode.\n" +
+			"  -noWatermarks : don't embed the commandline or version in reports\n" +
+			"  -cleanCsvOut <filename> : write a standard format CSV after event stripping with metadata stripped out.\n" +
+			"     Not available in bulk mode.\n" +
+			"  -noSmooth : disable smoothing on all graphs\n" +
+			"  -listSummaryTables: lists available summary tables from the current report XML\n" +
 			"\n" +
 			"Performance args:\n" +
-			"       -perfLog : output performance logging information\n" +
-			"       -graphThreads : use with -batchedGraphs to control the number of threads per CsvToSVG instance \n" +
-			"                       (default: PC core count/2)\n" +
-			"       -csvToSvgSequential : Run CsvToSvg sequentially\n" +
+			"  -perfLog : output performance logging information\n" +
+			"  -graphThreads : use with -batchedGraphs to control the number of threads per CsvToSVG instance \n" +
+			"                  (default: PC core count/2)\n" +
+			"  -csvToSvgSequential : Run CsvToSvg sequentially\n" +
 			"Deprecated performance args:\n" +
-			"       -csvToSvgProcesses : Use separate processes for csvToSVG instead of threads (slower)\n" +
-			"       -noBatchedGraphs : disable batched/multithreaded graph generation (use with -csvToSvgProcesses. Default is enabled)\n" +
+			"  -csvToSvgProcesses : Use separate processes for csvToSVG instead of threads (slower)\n" +
+			"  -noBatchedGraphs : disable batched/multithreaded graph generation (use with -csvToSvgProcesses. Default is enabled)\n" +
 			"\n" +
 			"Options to truncate or filter source data:\n" +
 			"Warning: these options disable Summary Table caching\n" +
-			"       -minx <frameNumber>\n" +
-			"       -maxx <frameNumber>\n" +
-			"       -beginEvent <event> : strip data before this event\n" +
-			"       -endEvent <event> : strip data after this event\n" +
-			"       -noStripEvents : if specified, don't strip out samples between excluded events from the stats\n" +
+			"  -minx <frameNumber>\n" +
+			"  -maxx <frameNumber>\n" +
+			"  -beginEvent <event> : strip data before this event\n" +
+			"  -endEvent <event> : strip data after this event\n" +
+			"  -noStripEvents : if specified, don't strip out samples between excluded events from the stats\n" +
 			"\n" +
 			"Optional bulk mode args: (use with -csvdir, -summaryTableCacheIn, -csvList, -prcList)\n" +
-			"       -recurse \n" +
-			"       -searchpattern <pattern>, e.g -searchpattern csvprofile*\n" +
-			"       -customTable <comma separated fields>\n" +
-			"       -customTableSort <comma separated field row sort order> (use with -customTable)\n" +
-			"       -noDetailedReports : skips individual report generation\n" +
-			"       -collateTable : writes a collated table in addition to the main one, merging by row sort\n" +
-			"       -collateTableOnly : as -collateTable, but doesn't write the standard summary table.\n" +
-			"       -emailTable : writes a condensed email-friendly table (see the 'condensed' summary table)\n" +
-			"       -csvTable : writes the summary table in CSV format instead of html\n" +
-			"       -summaryTableXML <XML filename>\n" +
-			"       -summaryTable <name> :\n" +
-			"           Selects a custom summary table type from the list in reportTypes.xml \n" +
-			"           (if not specified, 'default' will be used)\n" +
-			"       -condensedSummaryTable <name> :\n" +
-			"           Selects a custom condensed summary table type from the list in reportTypes.xml \n" +
-			"           (if not specified, 'condensed' will be used)\n" +
-			"       -summaryTableFilename <name> : use the specified filename for the summary table (instead of SummaryTable.html)\n" +
-			"       -metadataFilter <query> or <key0=value0,key1=value1...>: filters based on CSV metadata,\n" +
-			"           e.g \"platform=ps4 AND deviceprofile=ps4_60\" \n" +
-			"       -readAllStats : allows any CSV stat avg to appear in the summary table, not just those referenced in summaries\n" +
-			"       -showHiddenStats : shows stats which have been automatically hidden (typically duplicate csv unit stats)\n" +
-			"       -externalGraphs : enables external graphs (off by default)\n" +
-			"       -spreadsheetfriendly: outputs a single quote before non-numeric entries in summary tables\n" +
-			"       -noSummaryMinMax: don't make min/max columns for each stat in a condensed summary\n" +
-			"       -reverseTable [0|1]: Reverses the order of summary tables (set 0 to force off)\n" +
-			"       -scrollableTable [0|1]: makes the summary table scrollable, with frozen first rows and columns (set 0 to force off)\n" +
-			"       -colorizeTable [off|budget|auto]: selects the table colorization mode. If omitted, uses the default in the summary xml table if set.\n" +
-			"       -maxSummaryTableStringLength <n>: strings longer than this will get truncated\n" +
-			"       -allowDuplicateCSVs : doesn't remove duplicate CSVs (Note: can cause summary table cache file locking issues)\n" +
-			"       -requireMetadata : ignores CSVs without metadata\n" +
-			"       -listFiles : just list all files that pass the metadata query. Don't generate any reports.\n" +
-			"       -reportLinkRootPath <path> : Make report links relative to this\n" +
-			"       -csvLinkRootPath <path> : Make CSV file links relative to this\n" +
-			"       -linkTemplates : insert templates in place of relative links that can be replaced later, e.g {{LinkTemplate:Report:<CSV ID>}}\n" +
-			"       -weightByColumn : weight collated table averages by this column (overrides value specified in the report XML)\n" +
-			"       -noWeightedAvg : Don't use weighted averages for the collated table\n" +
-			"       -minFrameCount <n> : ignore CSVs without at least this number of valid frames\n" +
-			"       -maxFileAgeDays <n> : max file age in days. CSV or PRC files older than this will be ignored\n" +
-			"       -summaryTableStatThreshold <n> : stat/metric columns in the summarytable will be filtered out if all values < threshold\n" +
-			"       -summaryTableXmlSubst <find1>=<replace1>,<find2>=<replace2>... : replace summarytable XML row and filter entries\n" +
-			"       -transposeTable : write the summary tables transposed\n"+
-			"       -transposeCollatedTable : write the collated summary table transposed (disables min/max columns)\n" +
-            "       -addDiffRows : adds diff rows after the first two rows\n" + 
+			"  -recurse \n" +
+			"  -searchpattern <pattern>, e.g -searchpattern csvprofile*\n" +
+			"  -customTable <comma separated fields>\n" +
+			"  -customTableSort <comma separated field row sort order> (use with -customTable)\n" +
+			"  -noDetailedReports : skips individual report generation\n" +
+			"  -collateTable : writes a collated table in addition to the main one, merging by row sort\n" +
+			"  -collateTableOnly : as -collateTable, but doesn't write the standard summary table.\n" +
+			"  -emailTable : writes a condensed email-friendly table (see the 'condensed' summary table)\n" +
+			"  -csvTable : writes the summary table in CSV format instead of html\n" +
+			"  -summaryTableXML <XML filename>\n" +
+			"  -summaryTable <name> :\n" +
+			"     Selects a custom summary table type from the list in reportTypes.xml \n" +
+			"     (if not specified, 'default' will be used)\n" +
+			"  -condensedSummaryTable <name> :\n" +
+			"     Selects a custom condensed summary table type from the list in reportTypes.xml \n" +
+			"     (if not specified, 'condensed' will be used)\n" +
+			"  -summaryTableFilename <name> : use the specified filename for the summary table (instead of SummaryTable.html)\n" +
+			"  -metadataFilter <query> or <key0=value0,key1=value1...>: filters based on CSV metadata,\n" +
+			"     e.g \"platform=ps4 AND deviceprofile=ps4_60\" \n" +
+			"  -readAllStats : allows any CSV stat avg to appear in the summary table, not just those referenced in summaries\n" +
+			"  -showHiddenStats : shows stats which have been automatically hidden (typically duplicate csv unit stats)\n" +
+			"  -externalGraphs : enables external graphs (off by default)\n" +
+			"  -spreadsheetfriendly: outputs a single quote before non-numeric entries in summary tables\n" +
+			"  -noSummaryMinMax: don't make min/max columns for each stat in a condensed summary\n" +
+			"  -reverseTable [0|1]: Reverses the order of summary tables (set 0 to force off)\n" +
+			"  -scrollableTable [0|1]: makes the summary table scrollable, with frozen first rows and columns (set 0 to force off)\n" +
+			"  -colorizeTable [off|budget|auto]: selects the table colorization mode. If omitted, uses the default in the summary\n" +
+			"     xml table if set.\n" +
+			"  -maxSummaryTableStringLength <n>: strings longer than this will get truncated\n" +
+			"  -allowDuplicateCSVs : doesn't remove duplicate CSVs (Note: can cause summary table cache file locking issues)\n" +
+			"  -requireMetadata : ignores CSVs without metadata\n" +
+			"  -listFiles : just list all files that pass the metadata query. Don't generate any reports.\n" +
+			"  -reportLinkRootPath <path> : Make report links relative to this\n" +
+			"  -csvLinkRootPath <path> : Make CSV file links relative to this\n" +
+			"  -linkTemplates : insert templates in place of relative links that can be replaced later\n" +
+			"     e.g {{LinkTemplate:Report:<CSV ID>}}\n" +
+			"  -weightByColumn : weight collated table averages by this column (overrides value specified in the report XML)\n" +
+			"  -noWeightedAvg : Don't use weighted averages for the collated table\n" +
+			"  -minFrameCount <n> : ignore CSVs without at least this number of valid frames\n" +
+			"  -maxFileAgeDays <n> : max file age in days. CSV or PRC files older than this will be ignored\n" +
+			"  -summaryTableStatThreshold <n> : stat/metric columns in the summarytable will be filtered out if all values are\n" +
+			"     less than the threshold\n" +
+			"  -summaryTableXmlSubst <find1>=<replace1>,<find2>=<replace2>... : replace summarytable XML row and filter entries\n" +
+			"  -transposeTable : write the summary tables transposed\n"+
+			"  -transposeCollatedTable : write the collated summary table transposed (disables min/max columns)\n" +
+            "  -addDiffRows : adds diff rows after the first two rows\n" + 
 			"\n" +
 			"Optional Column Filters\n" +
-			"		-debugShowFilteredColumns : grays out filtered columns instead of removing them. Column tooltip will show filtered reason.\n" +
-			"		-hideMetadataColumns : filters out metadata columns from the table (excluding those used in row sort).\n" +
-			"		-onlyShowRegressedColumns : only shows columns where the most recent row group (see -regressionJoinRowsByName) is outside the given standard deviation threshold from the mean of the previous rows.\n" +
-			"			-regressionJoinRowsByName <statName> (optional) : the csv stat name to join rows by for aggregation. If not provided each row is treated distinctly and no aggregation occurs.\n" +
-			"			-regressionStdDevThreshold <n> (default = 2) : how many standard deviations from the mean the most recent row group must be to be considered a regression.\n" +
-			"			-regressionOutlierStdDevThreshold <n> (default = 4) : how many standard deviations from the mean a value must be to be considered an outlier. Used in the initial pass to remove outliers to better identify real regressions.\n" +
+			"  -debugShowFilteredColumns : grays out filtered columns instead of removing. Column tooltip will show filtered reason.\n" +
+			"  -hideMetadataColumns : filters out metadata columns from the table (excluding those used in row sort).\n" +
+			"\n" +
+			"Regression Column Filtering\n" +
+			"  -onlyShowRegressedColumns : enables regression filtering. Only shows columns where the most recent row group\n" +
+			"    (see -regressionJoinRowsByName) is outside the given stddiv threshold from the mean of the previous rows.\n" +
+			"  -regressionJoinRowsByName <statName> : a stat name to join rows by for aggregation. (default: no aggregation)\n"+
+			"  -regressionStdDevThreshold <n> (default = 2) : the stddiv threshold for filtering \n"+
+			"  -regressionOutlierStdDevThreshold <n> (default = 4) : stddiv threshold for outliers (these are ignored)\n"+
 			"\n" +
 			"Json serialization:\n" +
-			"       -summaryTableToJson <filename> : json filename to write summary table row data to\n" +
-			"       -summaryTableToJsonFastMode : exit after serializing json data (skips making summary tables)\n" +
-			"       -summaryTableToJsonWriteAllElementData : write all element data, including tooltips, flags\n" +
-			"       -summaryTableToJsonMetadataOnly : only write CsvMetadata elements to json\n" +
-			"       -summaryTableToJsonFileStream : use a file stream to write Json. Experimental but can avoid OOMs\n" +
-			"       -jsonToPrcs <json filename> : write PRCs. PRC files will be written to -summaryTableCache folder\n" +
+			"  -summaryTableToJson <filename> : json filename to write summary table row data to\n" +
+			"  -summaryTableToJsonFastMode : exit after serializing json data (skips making summary tables)\n" +
+			"  -summaryTableToJsonWriteAllElementData : write all element data, including tooltips, flags\n" +
+			"  -summaryTableToJsonMetadataOnly : only write CsvMetadata elements to json\n" +
+			"  -summaryTableToJsonFileStream : use a file stream to write Json. Experimental but can avoid OOMs\n" +
+			"  -jsonToPrcs <json filename> : write PRCs. PRC files will be written to -summaryTableCache folder\n" +
 			"\n" +
 			"Performance args for bulk mode:\n" +
-			"       -precacheCount <n> : number of CSV files to precache in the lookahead cache (0 for no precache)\n" +
-			"       -precacheThreads <n> : number of threads to use for the CSV lookahead cache (default 8)\n" +
-			"       -summaryTableCache <dir> : specifies a directory for summary table data to be cached.\n" +
-			"           This avoids processing csvs on subsequent runs when -noDetailedReports is specified\n" +
-			"           Note: Enables -readAllStats implicitly. \n" +
-			"       -summaryTableCacheInvalidate : regenerates summary table disk cache entries (ie write only)\n" +
-			"       -summaryTableCacheReadOnly : only read from the cache, never write\n" +
-			"       -summaryTableCachePurgeInvalid : Purges invalid PRCs from the cache folder\n" +
-			"       -summaryTableCacheIn <dir> : reads data directly from the summary table cache instead of from CSVs\n" +
-			"       -summaryTableCacheUseOnlyCsvID : only use the CSV ID for the summary table cacheID, ignoringthe report type hash\n" +
-			"            Use this if you want to avoid cache data being invalidated by report changes\n" +
-			"       -noCsvCacheFiles: disables usage of .csv.cache files. Cache files can be much faster if filtering on metadata\n" +
+			"  -precacheCount <n> : number of CSV files to precache in the lookahead cache (0 for no precache)\n" +
+			"  -precacheThreads <n> : number of threads to use for the CSV lookahead cache (default 8)\n" +
+			"  -summaryTableCache <dir> : specifies a directory for summary table data to be cached.\n" +
+			"     This avoids processing csvs on subsequent runs when -noDetailedReports is specified\n" +
+			"     Note: Enables -readAllStats implicitly. \n" +
+			"  -summaryTableCacheInvalidate : regenerates summary table disk cache entries (ie write only)\n" +
+			"  -summaryTableCacheReadOnly : only read from the cache, never write\n" +
+			"  -summaryTableCachePurgeInvalid : Purges invalid PRCs from the cache folder\n" +
+			"  -summaryTableCacheIn <dir> : reads data directly from the summary table cache instead of from CSVs\n" +
+			"  -summaryTableCacheUseOnlyCsvID : only use the CSV ID for the summary table cacheID, ignoringthe report type hash\n" +
+			"     Use this if you want to avoid cache data being invalidated by report changes\n" +
+			"  -noCsvCacheFiles: disables usage of .csv.cache files. Cache files can be much faster if filtering on metadata\n" +
 			"";
 		/*
 		"Note on custom tables:\n" +
