PerfReportTool - support writing PRC json to separate files, via -summaryTableToJsonSeparateFiles. Also add -summaryTableToJsonNoIndent option

diff --git a/Engine/Source/Programs/CSVTools/PerfReportTool/PerfReportTool.cs b/Engine/Source/Programs/CSVTools/PerfReportTool/PerfReportTool.cs
--- a/Engine/Source/Programs/CSVTools/PerfReportTool/PerfReportTool.cs
+++ b/Engine/Source/Programs/CSVTools/PerfReportTool/PerfReportTool.cs
@@ -26,7 +26,7 @@ namespace PerfReportTool
     class Version
     {
 		// Format: Major.Minor.Bugfix
-        private static string VersionString = "4.99.2";
+        private static string VersionString = "4.100.0";
 
         public static string Get() { return VersionString; }
     };
@@ -181,11 +181,13 @@ namespace PerfReportTool
 			"  -regressionOutlierStdDevThreshold <n> (default = 4) : stddiv threshold for outliers (these are ignored)\n"+
 			"\n" +
 			"Json serialization:\n" +
-			"  -summaryTableToJson <filename> : json filename to write summary table row data to\n" +
+			"  -summaryTableToJson <path> : path (usually a json filename) to write summary table row data to\n" +
+			"  -summaryTableToJsonSeparateFiles : writes separate files. -summaryTableToJson specifies the directory name\n" +
 			"  -summaryTableToJsonFastMode : exit after serializing json data (skips making summary tables)\n" +
 			"  -summaryTableToJsonWriteAllElementData : write all element data, including tooltips, flags\n" +
 			"  -summaryTableToJsonMetadataOnly : only write CsvMetadata elements to json\n" +
 			"  -summaryTableToJsonFileStream : use a file stream to write Json. Experimental but can avoid OOMs\n" +
+			"  -summaryTableToJsonNoIndent : don't indent json output files\n" +
 			"  -jsonToPrcs <json filename> : write PRCs. PRC files will be written to -summaryTableCache folder\n" +
 			"\n" +
 			"Performance args for bulk mode:\n" +
@@ -369,10 +371,14 @@ namespace PerfReportTool
 
 			// If we're outputting row data to json, create the dict
 			SummaryTableDataJsonWriteHelper summaryTableJsonHelper = null;
-			string summaryJsonOutFilename = GetArg("summaryTableToJson", null);
-			if (summaryJsonOutFilename != null)
+			string summaryJsonOutPath = GetArg("summaryTableToJson", null);
+			if (summaryJsonOutPath != null)
 			{
-				summaryTableJsonHelper = new SummaryTableDataJsonWriteHelper(summaryJsonOutFilename, GetBoolArg("summaryTableToJsonMetadataOnly"), GetBoolArg("summaryTableToJsonWriteAllElementData"));
+				summaryTableJsonHelper = new SummaryTableDataJsonWriteHelper(summaryJsonOutPath, 
+					GetBoolArg("summaryTableToJsonSeparateFiles"), 
+					GetBoolArg("summaryTableToJsonMetadataOnly"), 
+					GetBoolArg("summaryTableToJsonWriteAllElementData"),
+					!GetBoolArg("summaryTableToJsonNoIndent"));
 			}
 
 
@@ -632,7 +638,7 @@ namespace PerfReportTool
 
 			if (summaryTableJsonHelper != null)
 			{
-				summaryTableJsonHelper.WriteJsonFile(GetBoolArg("summaryTableToJsonFileStream"));
+				summaryTableJsonHelper.WriteToJson(GetBoolArg("summaryTableToJsonFileStream"));
 				perfLog.LogTiming("WriteSummaryDataJson");
 				if (GetBoolArg("summaryTableToJsonFastMode"))
 				{
diff --git a/Engine/Source/Programs/CSVTools/PerfReportTool/SummaryTableData.cs b/Engine/Source/Programs/CSVTools/PerfReportTool/SummaryTableData.cs
--- a/Engine/Source/Programs/CSVTools/PerfReportTool/SummaryTableData.cs
+++ b/Engine/Source/Programs/CSVTools/PerfReportTool/SummaryTableData.cs
@@ -14,11 +14,19 @@ namespace PerfSummaries
 
 	class SummaryTableDataJsonWriteHelper
 	{
-		public SummaryTableDataJsonWriteHelper(string InJsonFilename, bool bInCsvMetadataOnly, bool bInWriteAllElementData)
+		public SummaryTableDataJsonWriteHelper(string InJsonFilenameOrDirectory, bool bInSeparateFiles, bool bInCsvMetadataOnly, bool bInWriteAllElementData, bool bInWriteIndented=false)
 		{
-			JsonFilename = InJsonFilename;
+			if (bInSeparateFiles)
+			{
+				JsonDirectory = InJsonFilenameOrDirectory;
+			}
+			else
+			{
+				JsonFilename = InJsonFilenameOrDirectory;
+			}
 			bCsvMetadataOnly = bInCsvMetadataOnly;
 			bWriteAllElementData = bInWriteAllElementData;
+			bWriteIndented = bInWriteIndented;
 		}
 		public void AddRowData(SummaryTableRowData rowData)
 		{
@@ -28,39 +36,60 @@ namespace PerfSummaries
 				return;
 			}
 
-			Dict.Add(rowData.dict["csvid"].value, rowData.ToJsonDict(bCsvMetadataOnly, bWriteAllElementData));
+			Dictionary<string, dynamic> RowDict = rowData.ToJsonDict(bCsvMetadataOnly, bWriteAllElementData);
+			Dict.Add(rowData.dict["csvid"].value, RowDict);
 		}
 
-		public void WriteJsonFile(bool bSerializeToStream)
+		public void WriteToJson(bool bSerializeToStream)
 		{
-			Console.WriteLine("Writing summary table row data to json: " + JsonFilename);
+			if (JsonDirectory != null)
+			{
+				Console.WriteLine("Writing summary table row data to directory : " + JsonDirectory);
+				foreach (string csvId in Dict.Keys)
+				{
+					string filename = Path.Combine(JsonDirectory, csvId+".json");
+					WriteJsonFile(bSerializeToStream, filename, Dict[csvId], bWriteIndented);
+				}
+			}
+			else
+			{
+				Console.WriteLine("Writing summary table row data to file: " + JsonFilename);
+				WriteJsonFile(bSerializeToStream, JsonFilename, Dict, bWriteIndented);
+			}
 
-			JsonSerializerOptions options = new JsonSerializerOptions { WriteIndented = true };
+		}
+
+		private static void WriteJsonFile(bool bSerializeToStream, string filename, Dictionary<string, dynamic> dictToSerialize, bool bWriteIndented)
+		{
+			JsonSerializerOptions options = new JsonSerializerOptions { WriteIndented = bWriteIndented };
 
 			if (bSerializeToStream)
 			{
-				Task task = WriteJsonFileAsyncPrivate(options);
+				Task task = WriteJsonFileAsyncPrivate(options, filename, dictToSerialize);
 				task.Wait();
 			}
 			else
 			{
-				string jsonString = JsonSerializer.Serialize(Dict, options);
-				File.WriteAllText(JsonFilename, jsonString);
+				string jsonString = JsonSerializer.Serialize(dictToSerialize, options);
+				File.WriteAllText(filename, jsonString);
 			}
 		}
 
-		private async Task WriteJsonFileAsyncPrivate(JsonSerializerOptions options) 
+
+		private static async Task WriteJsonFileAsyncPrivate(JsonSerializerOptions options, string Filename, Dictionary<string, dynamic> DictToSerialize) 
 		{
 			// serialize JSON directly to a file
-			FileStream createStream = File.Create(JsonFilename);
-			await JsonSerializer.SerializeAsync(createStream, Dict);
+			FileStream createStream = File.Create(Filename);
+			await JsonSerializer.SerializeAsync(createStream, DictToSerialize);
 			createStream.Dispose(); // DisposeAsync doesn't exist in .Net 4.8 so we can't await it
 		}
 
 		Dictionary<string, dynamic> Dict = new Dictionary<string, dynamic>();
-		string JsonFilename;
+		string JsonFilename = null;
+		string JsonDirectory = null;
 		bool bCsvMetadataOnly;
 		bool bWriteAllElementData;
+		bool bWriteIndented;
 	}
 
 
