PerfReportTool: add -summaryTableXmlAppend arg for appending stats/metrics to the end of a summary table's filter list via the commandline. This allows us to append multiple stats (-summaryTableXmlSubst only supports one stat per substitution)

diff --git a/Engine/Source/Programs/CSVTools/PerfReportTool/PerfReportTool.cs b/Engine/Source/Programs/CSVTools/PerfReportTool/PerfReportTool.cs
--- a/Engine/Source/Programs/CSVTools/PerfReportTool/PerfReportTool.cs
+++ b/Engine/Source/Programs/CSVTools/PerfReportTool/PerfReportTool.cs
@@ -26,7 +26,7 @@ namespace PerfReportTool
     class Version
     {
 		// Format: Major.Minor.Bugfix
-        private static string VersionString = "4.100.0";
+        private static string VersionString = "4.101.0";
 
         public static string Get() { return VersionString; }
     };
@@ -165,6 +165,7 @@ namespace PerfReportTool
 			"  -summaryTableStatThreshold <n> : stat/metric columns in the summarytable will be filtered out if all values are\n" +
 			"     less than the threshold\n" +
 			"  -summaryTableXmlSubst <find1>=<replace1>,<find2>=<replace2>... : replace summarytable XML row and filter entries\n" +
+			"  -summaryTableXmlAppend <list,of,stats> : append these stats to the summary table's filter list\n" +
 			"  -transposeTable : write the summary tables transposed\n"+
 			"  -transposeCollatedTable : write the collated summary table transposed (disables min/max columns)\n" +
             "  -addDiffRows : adds diff rows after the first two rows\n" + 
@@ -354,7 +355,7 @@ namespace PerfReportTool
 			}
 
 			// Load the report + graph XML data
-			reportXML = new ReportXML(GetArg("graphxml", false), GetArg("reportxml", false), GetArg("reportxmlbasedir", false), GetArg("summaryTableXml", false), GetArg("summaryTableXmlSubst", false));
+			reportXML = new ReportXML(GetArg("graphxml", false), GetArg("reportxml", false), GetArg("reportxmlbasedir", false), GetArg("summaryTableXml", false), GetArg("summaryTableXmlSubst", false), GetArg("summaryTableXmlAppend", false));
 
 			if (GetBoolArg("listSummaryTables"))
 			{
diff --git a/Engine/Source/Programs/CSVTools/PerfReportTool/ReportXml.cs b/Engine/Source/Programs/CSVTools/PerfReportTool/ReportXml.cs
--- a/Engine/Source/Programs/CSVTools/PerfReportTool/ReportXml.cs
+++ b/Engine/Source/Programs/CSVTools/PerfReportTool/ReportXml.cs
@@ -67,7 +67,7 @@ namespace PerfReportTool
 			return false;
 		}
 
-		public ReportXML(string graphXMLFilenameIn, string reportXMLFilenameIn, string baseXMLDirectoryOverride, string additionalSummaryTableXmlFilename, string summaryTableXmlSubstStr)
+		public ReportXML(string graphXMLFilenameIn, string reportXMLFilenameIn, string baseXMLDirectoryOverride, string additionalSummaryTableXmlFilename, string summaryTableXmlSubstStr, string summaryTableXmlAppendStr)
 		{
 			
 			string location = System.Reflection.Assembly.GetEntryAssembly().Location.ToLower();
@@ -274,11 +274,17 @@ namespace PerfReportTool
 					}
 				}
 
+				string[] appendList = null;
+				if ( summaryTableXmlAppendStr != null )
+				{
+					appendList = summaryTableXmlAppendStr.Split(',');
+				}
+
 
 				summaryTables = new Dictionary<string, SummaryTableInfo>();
 				foreach (XElement summaryElement in summaryTablesElement.Elements("summaryTable"))
 				{
-					SummaryTableInfo table = new SummaryTableInfo(summaryElement, substitutionsDict);
+					SummaryTableInfo table = new SummaryTableInfo(summaryElement, substitutionsDict, appendList);
 					summaryTables.Add(summaryElement.Attribute("name").Value.ToLower(), table);
 				}
 			}
diff --git a/Engine/Source/Programs/CSVTools/PerfReportTool/SummaryTable.cs b/Engine/Source/Programs/CSVTools/PerfReportTool/SummaryTable.cs
--- a/Engine/Source/Programs/CSVTools/PerfReportTool/SummaryTable.cs
+++ b/Engine/Source/Programs/CSVTools/PerfReportTool/SummaryTable.cs
@@ -67,7 +67,7 @@ namespace PerfSummaries
 
 	class SummaryTableInfo
 	{
-		public SummaryTableInfo(XElement tableElement, Dictionary<string,string> substitutionsDict)
+		public SummaryTableInfo(XElement tableElement, Dictionary<string,string> substitutionsDict, string[] appendList)
 		{
 			XAttribute rowSortAt = tableElement.Attribute("rowSort");
 			if (rowSortAt != null)
@@ -88,6 +88,11 @@ namespace PerfSummaries
 				ApplySubstitutionsToList(columnFilterList, substitutionsDict);
 			}
 
+			if (appendList != null)
+			{
+				columnFilterList.AddRange(appendList);
+			}
+
 			bReverseSortRows = tableElement.GetSafeAttibute<bool>("reverseSortRows", false);
 			bScrollableFormatting = tableElement.GetSafeAttibute<bool>("scrollableFormatting", false);
 
