PerfReportTool - Add -summaryTableXmlRowSortAppend arg to allow appending values to a summary table's row sort list via the commandline

diff --git a/Engine/Source/Programs/CSVTools/PerfReportTool/PerfReportTool.cs b/Engine/Source/Programs/CSVTools/PerfReportTool/PerfReportTool.cs
--- a/Engine/Source/Programs/CSVTools/PerfReportTool/PerfReportTool.cs
+++ b/Engine/Source/Programs/CSVTools/PerfReportTool/PerfReportTool.cs
@@ -26,7 +26,7 @@ namespace PerfReportTool
     class Version
     {
 		// Format: Major.Minor.Bugfix
-        private static string VersionString = "4.101.0";
+        private static string VersionString = "4.102.0";
 
         public static string Get() { return VersionString; }
     };
@@ -166,7 +166,8 @@ namespace PerfReportTool
 			"     less than the threshold\n" +
 			"  -summaryTableXmlSubst <find1>=<replace1>,<find2>=<replace2>... : replace summarytable XML row and filter entries\n" +
 			"  -summaryTableXmlAppend <list,of,stats> : append these stats to the summary table's filter list\n" +
-			"  -transposeTable : write the summary tables transposed\n"+
+			"  -summaryTableXmlRowSortAppend <list,of,stats> : append these stats to the summary table's row sort list\n" +
+			"  -transposeTable : write the summary tables transposed\n" +
 			"  -transposeCollatedTable : write the collated summary table transposed (disables min/max columns)\n" +
             "  -addDiffRows : adds diff rows after the first two rows\n" + 
 			"\n" +
@@ -355,7 +356,15 @@ namespace PerfReportTool
 			}
 
 			// Load the report + graph XML data
-			reportXML = new ReportXML(GetArg("graphxml", false), GetArg("reportxml", false), GetArg("reportxmlbasedir", false), GetArg("summaryTableXml", false), GetArg("summaryTableXmlSubst", false), GetArg("summaryTableXmlAppend", false));
+			reportXML = new ReportXML(
+				GetArg("graphxml", false), 
+				GetArg("reportxml", false), 
+				GetArg("reportxmlbasedir", false), 
+				GetArg("summaryTableXml", false), 
+				GetArg("summaryTableXmlSubst", false), 
+				GetArg("summaryTableXmlAppend", false),
+				GetArg("summaryTableXmlRowSortAppend", false)
+				);
 
 			if (GetBoolArg("listSummaryTables"))
 			{
diff --git a/Engine/Source/Programs/CSVTools/PerfReportTool/ReportXml.cs b/Engine/Source/Programs/CSVTools/PerfReportTool/ReportXml.cs
--- a/Engine/Source/Programs/CSVTools/PerfReportTool/ReportXml.cs
+++ b/Engine/Source/Programs/CSVTools/PerfReportTool/ReportXml.cs
@@ -67,9 +67,16 @@ namespace PerfReportTool
 			return false;
 		}
 
-		public ReportXML(string graphXMLFilenameIn, string reportXMLFilenameIn, string baseXMLDirectoryOverride, string additionalSummaryTableXmlFilename, string summaryTableXmlSubstStr, string summaryTableXmlAppendStr)
+		public ReportXML(
+			string graphXMLFilenameIn, 
+			string reportXMLFilenameIn, 
+			string baseXMLDirectoryOverride, 
+			string additionalSummaryTableXmlFilename, 
+			string summaryTableXmlSubstStr, 
+			string summaryTableXmlAppendStr,
+			string summaryTableXmlRowSortAppendStr )
 		{
-			
+
 			string location = System.Reflection.Assembly.GetEntryAssembly().Location.ToLower();
 			string baseDirectory = location.Replace("perfreporttool.exe", "");
 
@@ -280,11 +287,17 @@ namespace PerfReportTool
 					appendList = summaryTableXmlAppendStr.Split(',');
 				}
 
+				string[] rowSortAppendList = null;
+				if (summaryTableXmlRowSortAppendStr != null)
+				{
+					rowSortAppendList = summaryTableXmlRowSortAppendStr.Split(',');
+				}
+
 
 				summaryTables = new Dictionary<string, SummaryTableInfo>();
 				foreach (XElement summaryElement in summaryTablesElement.Elements("summaryTable"))
 				{
-					SummaryTableInfo table = new SummaryTableInfo(summaryElement, substitutionsDict, appendList);
+					SummaryTableInfo table = new SummaryTableInfo(summaryElement, substitutionsDict, appendList, rowSortAppendList);
 					summaryTables.Add(summaryElement.Attribute("name").Value.ToLower(), table);
 				}
 			}
diff --git a/Engine/Source/Programs/CSVTools/PerfReportTool/SummaryTable.cs b/Engine/Source/Programs/CSVTools/PerfReportTool/SummaryTable.cs
--- a/Engine/Source/Programs/CSVTools/PerfReportTool/SummaryTable.cs
+++ b/Engine/Source/Programs/CSVTools/PerfReportTool/SummaryTable.cs
@@ -67,7 +67,7 @@ namespace PerfSummaries
 
 	class SummaryTableInfo
 	{
-		public SummaryTableInfo(XElement tableElement, Dictionary<string,string> substitutionsDict, string[] appendList)
+		public SummaryTableInfo(XElement tableElement, Dictionary<string,string> substitutionsDict, string[] appendList, string[] rowSortAppendList)
 		{
 			XAttribute rowSortAt = tableElement.Attribute("rowSort");
 			if (rowSortAt != null)
@@ -75,6 +75,11 @@ namespace PerfSummaries
 				rowSortList.AddRange(rowSortAt.Value.Split(','));
 				ApplySubstitutionsToList(rowSortList, substitutionsDict);
 			}
+			if (rowSortAppendList != null)
+			{
+				rowSortList.AddRange(rowSortAppendList);
+			}
+
 			XAttribute weightByColumnAt = tableElement.Attribute("weightByColumn");
 			if (weightByColumnAt != null)
 			{
@@ -1749,7 +1754,7 @@ namespace PerfSummaries
 					if (columnLookup.ContainsKey(s.ToLower()))
 					{
 						SummaryTableColumn column = columnLookup[s.ToLower()];
-						key += "{" + column.GetStringValue(i,false,"0.0000000000") + "}";
+						key += "{" + column.GetStringValue(i,false,"0000000000.0000000000") + "}";
 					}
 					else
 					{
