[PerfReportTool] Removed gamethread histogram from the summary table. This works better as part of the report (see HistogramSummary.cs)

diff --git a/Engine/Source/Programs/CSVTools/PerfReportTool/PerfReportTool.cs b/Engine/Source/Programs/CSVTools/PerfReportTool/PerfReportTool.cs
--- a/Engine/Source/Programs/CSVTools/PerfReportTool/PerfReportTool.cs
+++ b/Engine/Source/Programs/CSVTools/PerfReportTool/PerfReportTool.cs
@@ -26,7 +26,7 @@ namespace PerfReportTool
     class Version
     {
 		// Format: Major.Minor.Bugfix
-        private static string VersionString = "4.102.0";
+        private static string VersionString = "4.103.0";
 
         public static string Get() { return VersionString; }
     };
diff --git a/Engine/Source/Programs/CSVTools/PerfReportTool/Summaries/FPSChartSummary.cs b/Engine/Source/Programs/CSVTools/PerfReportTool/Summaries/FPSChartSummary.cs
--- a/Engine/Source/Programs/CSVTools/PerfReportTool/Summaries/FPSChartSummary.cs
+++ b/Engine/Source/Programs/CSVTools/PerfReportTool/Summaries/FPSChartSummary.cs
@@ -259,22 +259,6 @@ namespace PerfSummaries
 					rowData.Add(SummaryTableElement.Type.SummaryTableMetric, columnName, column.Value, column.ColorThresholds);
 				}
 				rowData.Add(SummaryTableElement.Type.SummaryTableMetric, "TargetFPS", (double)fps);
-
-				// Write out gamethread histogram to the summary for 0.25x increments up to 2x
-				StatSamples gamethreadtimeStat = csvStats.GetStat("GameThreadtime");
-				if (gamethreadtimeStat != null)
-				{
-					float targetTimeInMS = 1000.0f / fps;
-					float redThreshold = 25.0f;
-					float[] targetTimeMultipliers = new float[] { 0.75f, 1.00f, 1.25f, 1.50f };
-					foreach(float targetTimeMultiplier in targetTimeMultipliers)
-					{
-						float percentageOver = 100.0f * gamethreadtimeStat.GetRatioOfFramesOverBudget(targetTimeInMS * targetTimeMultiplier);
-						ColourThresholdList colorThresholdList = new ColourThresholdList(redThreshold, redThreshold * 0.66, redThreshold * 0.33, 0.0);
-
-						rowData.Add(SummaryTableElement.Type.SummaryTableMetric, "GameThread >" + targetTimeMultiplier + "x %", (double)percentageOver, colorThresholdList);
-					}
-				}
 			}
 
 			// Output HTML
