PerfReportTool - Fix divide by 0 issue in MVP calculation.

diff --git a/Engine/Source/Programs/CSVTools/PerfReportTool/Summaries/FPSChartSummary.cs b/Engine/Source/Programs/CSVTools/PerfReportTool/Summaries/FPSChartSummary.cs
--- a/Engine/Source/Programs/CSVTools/PerfReportTool/Summaries/FPSChartSummary.cs
+++ b/Engine/Source/Programs/CSVTools/PerfReportTool/Summaries/FPSChartSummary.cs
@@ -117,9 +117,17 @@ namespace PerfSummaries
 			totalHitchTime -= (hitchCount * hitchThreshold);
 			outData.HitchTimePercent = (float)(totalHitchTime / totalFrametime) * 100.0f;
 
+			// If frame count is low enough this can get rounded down to 0
 			int TotalTargetFrames = (int)((double)fps * (TotalSeconds));
-			int MissedFrames = Math.Max(TotalTargetFrames - frameTimes.Count, 0);
-			outData.MVP = (((float)MissedFrames * 100.0f) / (float)TotalTargetFrames);
+			if (TotalTargetFrames > 0)
+			{
+				int MissedFrames = Math.Max(TotalTargetFrames - frameTimes.Count, 0);
+				outData.MVP = (((float)MissedFrames * 100.0f) / (float)TotalTargetFrames);
+			}
+			else
+			{
+				outData.MVP = 0;
+			}
 			return outData;
 		}
 
