CsvConvert 1.05 - automatically repair null metadata if we're adding metadata (this is a common use-case)
CsvStats - fix bug where ReadFileFromLines wouldn't close the file (breaking -inplace mode in CsvConvert)

diff --git a/Engine/Source/Programs/CSVTools/CsvConvert/CsvConvert.cs b/Engine/Source/Programs/CSVTools/CsvConvert/CsvConvert.cs
--- a/Engine/Source/Programs/CSVTools/CsvConvert/CsvConvert.cs
+++ b/Engine/Source/Programs/CSVTools/CsvConvert/CsvConvert.cs
@@ -13,7 +13,7 @@ namespace CSVInfo
 {
     class Version
     {
-        private static string VersionString = "1.04";
+        private static string VersionString = "1.05";
 
         public static string Get() { return VersionString; }
     };
@@ -187,6 +187,17 @@ namespace CSVInfo
 
 			// Set metadata if requested
 			string setMetadataFile = GetArg("setMetadataFile",null);
+			string SetMetadataStr = GetArg("setMetadata");
+
+			// Make sure we have metadata if we're setting it
+			if (csvStats.metaData == null && (setMetadataFile != null || SetMetadataStr != null) )
+			{
+				Console.WriteLine("File has null metadata. Creating...");
+				csvStats.metaData = new CsvMetadata();
+				csvStats.GenerateCsvIDIfMissing();
+			}
+
+
 			if (setMetadataFile != null)
 			{
 				if (csvStats.metaData == null)
@@ -207,7 +218,6 @@ namespace CSVInfo
 				}
 			}
 
-			string SetMetadataStr = GetArg("setMetadata");
 			if (SetMetadataStr != null && SetMetadataStr.Length > 0)
 			{
 				bool bModified = false;
diff --git a/Engine/Source/Programs/CSVTools/CsvStats/CsvStats.cs b/Engine/Source/Programs/CSVTools/CsvStats/CsvStats.cs
--- a/Engine/Source/Programs/CSVTools/CsvStats/CsvStats.cs
+++ b/Engine/Source/Programs/CSVTools/CsvStats/CsvStats.cs
@@ -566,6 +566,7 @@ namespace CSVStats
 			}
 			if (bIsUnicode)
 			{
+				reader.Close();
 				reader = new StreamReader(filename, Encoding.Unicode, true);
 			}
 			else
@@ -582,6 +583,7 @@ namespace CSVStats
 				}
 			}
 
+			reader.Close();
 			return lines.ToArray();
 		}
 
