PerfReportTool 4.104 - Adjust auto-colorization to reduce the color difference when the difference between min/max is very small. Also add a threshold to disable colorization altogether.

diff --git a/Engine/Source/Programs/CSVTools/PerfReportTool/PerfReportTool.cs b/Engine/Source/Programs/CSVTools/PerfReportTool/PerfReportTool.cs
--- a/Engine/Source/Programs/CSVTools/PerfReportTool/PerfReportTool.cs
+++ b/Engine/Source/Programs/CSVTools/PerfReportTool/PerfReportTool.cs
@@ -26,7 +26,7 @@ namespace PerfReportTool
     class Version
     {
 		// Format: Major.Minor.Bugfix
-        private static string VersionString = "4.103.0";
+        private static string VersionString = "4.104.0";
 
         public static string Get() { return VersionString; }
     };
diff --git a/Engine/Source/Programs/CSVTools/PerfReportTool/SummaryTable.cs b/Engine/Source/Programs/CSVTools/PerfReportTool/SummaryTable.cs
--- a/Engine/Source/Programs/CSVTools/PerfReportTool/SummaryTable.cs
+++ b/Engine/Source/Programs/CSVTools/PerfReportTool/SummaryTable.cs
@@ -710,11 +710,26 @@ namespace PerfSummaries
 				return;
 			}
 
+			double averageValue = totalValue / validCount;
+			double range = maxValue - minValue;
+
+			// Disable colorization where values are very similar
+			// The range has to be outside 0.25% of the average and >0.01 to get colorized
+			double colorizationRangeThreshold = Math.Max( Math.Abs(averageValue) * 0.0025, 0.01 ); 
+			if (range < colorizationRangeThreshold)
+			{
+				return;
+			}
+
+			// Adjust Min/Max value to ensure close values are not just binary red/green. If min/max is within 1% of the average or 0.02, adjust accordingly
+			double minColorizationRangeExtent = Math.Max( Math.Abs(averageValue) * 0.01, 0.02); 
+			maxValue = Math.Max(maxValue, averageValue + minColorizationRangeExtent);
+			minValue = Math.Min(minValue, averageValue - minColorizationRangeExtent);
+
 			Colour green = Colour.Green;
 			Colour yellow = Colour.Yellow;
 			Colour red = Colour.Red;
 
-			double averageValue = totalValue / validCount; // TODO: Weighted average 
 			colourThresholdOverride = new ColourThresholdList();
 			colourThresholdOverride.Add(new ThresholdInfo(minValue, (autoColorizeMode == AutoColorizeMode.HighIsBad) ? green : red));
 			colourThresholdOverride.Add(new ThresholdInfo(averageValue, yellow));
