PerfReportTool - avoid div/0 in autocolorization

diff --git a/Engine/Source/Programs/CSVTools/PerfReportTool/PerfReportTool.cs b/Engine/Source/Programs/CSVTools/PerfReportTool/PerfReportTool.cs
--- a/Engine/Source/Programs/CSVTools/PerfReportTool/PerfReportTool.cs
+++ b/Engine/Source/Programs/CSVTools/PerfReportTool/PerfReportTool.cs
@@ -26,7 +26,7 @@ namespace PerfReportTool
     class Version
     {
 		// Format: Major.Minor.Bugfix
-        private static string VersionString = "4.104.0";
+        private static string VersionString = "4.104.1";
 
         public static string Get() { return VersionString; }
     };
diff --git a/Engine/Source/Programs/CSVTools/PerfReportTool/SummaryTable.cs b/Engine/Source/Programs/CSVTools/PerfReportTool/SummaryTable.cs
--- a/Engine/Source/Programs/CSVTools/PerfReportTool/SummaryTable.cs
+++ b/Engine/Source/Programs/CSVTools/PerfReportTool/SummaryTable.cs
@@ -688,8 +688,8 @@ namespace PerfSummaries
 			colourThresholds = new List<ColourThresholdList>();
 			double maxValue = -double.MaxValue;
 			double minValue = double.MaxValue;
-			double totalValue = 0.0f;
-			double validCount = 0.0f;
+			double totalValue = 0.0;
+			double validCount = 0.0;
 			for (int i = 0; i < doubleValues.Count; i++)
 			{
 				if (IsDiffRow(i))
@@ -702,10 +702,10 @@ namespace PerfSummaries
 					maxValue = Math.Max(val, maxValue);
 					minValue = Math.Min(val, minValue);
 					totalValue += val;
-					validCount += 1.0f;
+					validCount += 1.0;
 				}
 			}
-			if (minValue == maxValue)
+			if (minValue == maxValue || validCount == 0.0)
 			{
 				return;
 			}
