PerfReportTool 4.105.0: Add support for sorting columns by diff and also filtering out columns based on a diff value threshold.
Adds the following params:
  -sortColumnsByDiff : sorts columns by the max of its diff scores (use with -addDiffRows)
  Note: Diff score corresponds to the value of a column's diff row; the sign is reversed if LowIsBad

  -columnDiffDisplayThreshold <value> : if specified, hides columns with max abs diff value below this threshold

Also: added some helpers to summary tables to find aggregate columns . Added a isCountColumn accessor to avoid string compares

diff --git a/Engine/Source/Programs/CSVTools/PerfReportTool/PerfReportTool.cs b/Engine/Source/Programs/CSVTools/PerfReportTool/PerfReportTool.cs
--- a/Engine/Source/Programs/CSVTools/PerfReportTool/PerfReportTool.cs
+++ b/Engine/Source/Programs/CSVTools/PerfReportTool/PerfReportTool.cs
@@ -26,7 +26,7 @@ namespace PerfReportTool
     class Version
     {
 		// Format: Major.Minor.Bugfix
-        private static string VersionString = "4.104.1";
+        private static string VersionString = "4.105.0";
 
         public static string Get() { return VersionString; }
     };
@@ -169,7 +169,10 @@ namespace PerfReportTool
 			"  -summaryTableXmlRowSortAppend <list,of,stats> : append these stats to the summary table's row sort list\n" +
 			"  -transposeTable : write the summary tables transposed\n" +
 			"  -transposeCollatedTable : write the collated summary table transposed (disables min/max columns)\n" +
-            "  -addDiffRows : adds diff rows after the first two rows\n" + 
+            "  -addDiffRows : adds diff rows after the first two rows\n" +
+			"  -sortColumnsByDiff : sorts columns by the max of its diff scores (use with -addDiffRows)\n" +
+			"      Note: Diff score corresponds to the value of a column's diff row; the sign is reversed if LowIsBad\n" +
+			"  -columnDiffDisplayThreshold <value> : if specified, hides columns with max diff value below this threshold\n" +
 			"\n" +
 			"Optional Column Filters\n" +
 			"  -debugShowFilteredColumns : grays out filtered columns instead of removing. Column tooltip will show filtered reason.\n" +
@@ -820,10 +823,9 @@ namespace PerfReportTool
 				filteredTable.ApplyDisplayNameMapping(statDisplaynameMapping);
 				string VersionString = GetBoolArg("noWatermarks") ? "" : Version.Get();
 				string summaryTitle = GetArg("summaryTitle", null);
-
 				if (GetBoolArg("addDiffRows"))
 				{
-					filteredTable.AddDiffRows();
+					filteredTable.AddDiffRows(GetBoolArg("sortColumnsByDiff"), GetFloatArg("columnDiffDisplayThreshold", 0.0f));
 				}
 
 				// Run again to add format info for any new columns that were added (eg. count).
diff --git a/Engine/Source/Programs/CSVTools/PerfReportTool/SummaryTable.cs b/Engine/Source/Programs/CSVTools/PerfReportTool/SummaryTable.cs
--- a/Engine/Source/Programs/CSVTools/PerfReportTool/SummaryTable.cs
+++ b/Engine/Source/Programs/CSVTools/PerfReportTool/SummaryTable.cs
@@ -247,6 +247,15 @@ namespace PerfSummaries
 		LowIsBad,
 	};
 
+	enum ColumnAggregateType
+	{
+		None,
+		Avg,
+		Min,
+		Max
+	};
+
+
 	class SummaryTableColumnFormatInfo
 	{
 		public SummaryTableColumnFormatInfo()
@@ -342,10 +351,14 @@ namespace PerfSummaries
 		public string displayName;
 		public bool isRowWeightColumn = false;
 		public bool hasDiffRows = false;
+		public bool isCountColumn = false;
 		// Column header tooltip. Displayed when hovering over the header.
 		public string tooltip = null;
 		// Multiplied with the colour of each cell to modify the colour (including header).
 		public Colour columnColourModifier = null;
+		public ColumnAggregateType aggregateType = ColumnAggregateType.None;
+		public SummaryTableColumn aggregateBaseColumn = null; // For avg/min/max columns, this will point back to the base(avg) column
+
 		List<double> doubleValues = new List<double>();
 		List<string> stringValues = new List<string>();
 		List<string> toolTips = new List<string>();
@@ -365,9 +378,12 @@ namespace PerfSummaries
 			SummaryTableElement.Type inElementType,
 			SummaryTableColumnFormatInfo inFormatInfo = null,
 			string inTooltip = null,
-			Colour inColumnColourModifier = null)
+			Colour inColumnColourModifier = null,
+			ColumnAggregateType inAggregateType = ColumnAggregateType.None,
+			SummaryTableColumn inAggregateBaseColumn = null,
+			bool bInIsCountColumn = false)
 		{
-			name = inName;
+			name = SummaryTableColumn.getAggregateTypePrefix(inAggregateType) + inName;
 			isNumeric = inIsNumeric;
 			displayName = inDisplayName;
 			isRowWeightColumn = inIsRowWeightColumn;
@@ -375,7 +391,30 @@ namespace PerfSummaries
 			formatInfo = inFormatInfo;
 			tooltip = inTooltip;
 			columnColourModifier = inColumnColourModifier;
+			aggregateType = inAggregateType;
+			aggregateBaseColumn = inAggregateBaseColumn;
+			isCountColumn = bInIsCountColumn;
+			if (inAggregateType == ColumnAggregateType.Avg && aggregateBaseColumn == null)
+			{
+				aggregateBaseColumn = this;
+			}
+		}
+
+		public static string getAggregateTypePrefix(ColumnAggregateType aggregateType)
+		{
+			switch (aggregateType)
+			{
+				case ColumnAggregateType.Avg:
+					return "Avg ";
+				case ColumnAggregateType.Min:
+					return "Min ";
+				case ColumnAggregateType.Max:
+					return "Max ";
+				default:
+					return "";
+			}
 		}
+
 		public SummaryTableColumn Clone()
 		{
 			SummaryTableColumn newColumn = new SummaryTableColumn(name, isNumeric, displayName, isRowWeightColumn, elementType, formatInfo, tooltip, columnColourModifier);
@@ -417,7 +456,7 @@ namespace PerfSummaries
 			List<string> newToolTips = new List<string>(toolTips.Count > 0 ? newCount : 0);
 			List<ColourThresholdList> newColourThresholds = new List<ColourThresholdList>(colourThresholds.Count > 0 ? newCount : 0);
 
-			bool bComputeDiff = isNumeric && name != "Count";
+			bool bComputeDiff = isNumeric && !isCountColumn;
 
 			// Add diff rows to each of the arrays
 			for (int i = 0; i < doubleValues.Count; i++)
@@ -478,6 +517,40 @@ namespace PerfSummaries
 			return ((rowIndex - 2) % 2) == 0;
 		}
 
+		// Computes a score (significance indicator) for a column based on its diff values. This takes into account the max value. If LowIsBad for this column then the sign is reversed
+		public double GetDiffScore()
+		{
+			if (!hasDiffRows || !isNumeric)
+			{
+				return 0.0;
+			}
+			bool bLowIsBad = formatInfo != null && formatInfo.autoColorizeMode == AutoColorizeMode.LowIsBad;
+
+			// Find the max of all diff values for this column. If LowIsBad then we reverse the sign
+			double maxDiffScore = double.MinValue;
+			for (int diffRowIndex = 2; diffRowIndex < GetCount(); diffRowIndex += 2)
+			{
+				double diffValue = GetValue(diffRowIndex);
+				maxDiffScore = Math.Max(maxDiffScore, bLowIsBad ? -diffValue : diffValue);
+			}
+			return maxDiffScore;
+		}
+
+		// Computes the max of the abs diff values for a column
+		public double GetMaxAbsDiff()
+		{
+			if (!hasDiffRows || !isNumeric)
+			{
+				return 0.0;
+			}
+			double maxAbsDiff = double.MinValue;
+			for (int diffRowIndex = 2; diffRowIndex < GetCount(); diffRowIndex += 2)
+			{
+				maxAbsDiff = Math.Max(maxAbsDiff, Math.Abs(GetValue(diffRowIndex)));
+			}
+			return maxAbsDiff;
+		}
+
 		public string GetDisplayName(string hideStatPrefix=null, bool bAddStatCategorySeparatorSpaces = true, bool bGreyOutStatCategories = false)
 		{
 			if (displayName != null)
@@ -970,7 +1043,7 @@ namespace PerfSummaries
 				}
 			}
 
-			newColumns.Add(new SummaryTableColumn("Count", true, null, false, SummaryTableElement.Type.ToolMetadata));
+			newColumns.Add(new SummaryTableColumn("Count", true, null, false, SummaryTableElement.Type.ToolMetadata, bInIsCountColumn:true));
 			int countColumnIndex = newColumns.Count - 1;
 
 			int numericColumnStartIndex = newColumns.Count;
@@ -981,11 +1054,12 @@ namespace PerfSummaries
 				if (column.DisplayAggregates() && !collateByColumns.Contains(column))
 				{
 					srcToDestBaseColumnIndex.Add(newColumns.Count);
-					newColumns.Add(new SummaryTableColumn("Avg " + column.name, true, null, false, column.elementType, column.formatInfo, column.tooltip, column.columnColourModifier));
+					SummaryTableColumn avgColumn = new SummaryTableColumn(column.name, true, null, false, column.elementType, column.formatInfo, column.tooltip, column.columnColourModifier, ColumnAggregateType.Avg);
+					newColumns.Add(avgColumn);
 					if (addMinMaxColumns)
 					{
-						newColumns.Add(new SummaryTableColumn("Min " + column.name, true, null, false, column.elementType, column.formatInfo, column.tooltip, column.columnColourModifier));
-						newColumns.Add(new SummaryTableColumn("Max " + column.name, true, null, false, column.elementType, column.formatInfo, column.tooltip, column.columnColourModifier));
+						newColumns.Add(new SummaryTableColumn(column.name, true, null, false, column.elementType, column.formatInfo, column.tooltip, column.columnColourModifier, ColumnAggregateType.Min, avgColumn));
+						newColumns.Add(new SummaryTableColumn(column.name, true, null, false, column.elementType, column.formatInfo, column.tooltip, column.columnColourModifier, ColumnAggregateType.Max, avgColumn));
 					}
 				}
 				else
@@ -1117,6 +1191,31 @@ namespace PerfSummaries
 			return newTable;
 		}
 
+		// Finds a particular aggregate column corresponding to a specified column
+		private SummaryTableColumn GetAggregateColumn(SummaryTableColumn inColumn, ColumnAggregateType aggregateType)
+		{
+			// Add the avg column and the corresponding min/max columns
+			if (inColumn.aggregateType == ColumnAggregateType.None)
+			{
+				throw new Exception("Column isn't an aggregate column");
+			}
+
+			// Aggregate columns keep a ref to the avg/base column so use that if appropriate
+			if (aggregateType == ColumnAggregateType.Avg && inColumn.aggregateBaseColumn != null)
+			{
+				return inColumn.aggregateBaseColumn;
+			}
+
+			string prefix = GetBaseStatPrefix(inColumn.name);
+			string lookupKey = ( SummaryTableColumn.getAggregateTypePrefix(aggregateType) + inColumn.name.Substring(prefix.Length) ).ToLower();
+			if ( !columnLookup.TryGetValue(lookupKey, out SummaryTableColumn outColumn) )
+			{
+				throw new Exception("Aggregate column "+lookupKey+" not found!");
+			}
+			return outColumn;
+		}
+
+
 		public SummaryTable SortAndFilter(string customFilter, string customRowSort = "buildversion,deviceprofile", bool bReverseSort = false, string weightByColumnName = null)
 		{
 			return SortAndFilter(customFilter.Split(',').ToList(), customRowSort.Split(',').ToList(), bReverseSort, weightByColumnName);
@@ -1241,7 +1340,7 @@ namespace PerfSummaries
 			return newTable;
 		}
 
-		public void AddDiffRows()
+		public void AddDiffRows(bool bSortColumnsByDiff, double columnDiffDisplayThreshold)
 		{
 			for (int i=0; i<columns.Count; i++)
 			{
@@ -1249,6 +1348,15 @@ namespace PerfSummaries
 			}
 			rowCount += rowCount - 1;
 
+			if ( columnDiffDisplayThreshold > 0.0 )
+			{
+				FilterColumnsByDiffThreshold(columnDiffDisplayThreshold);
+			}
+
+			if ( bSortColumnsByDiff )
+			{
+				SortColumnsByDiffRows();
+			}
 			// Just set rowWeightings to null for now. We shouldn't need it, since we should have already collated by this point
 			rowWeightings = null;
 		}
@@ -1269,6 +1377,13 @@ namespace PerfSummaries
 				}
 			}
 		}
+
+		public static string GetBaseStatPrefix(string inName)
+		{
+			GetBaseStatNameWithPrefixAndSuffix(inName, out string prefix, out _);
+			return prefix;
+		}
+
 		public static string GetBaseStatName(string inName)
 		{
 			return GetBaseStatNameWithPrefixAndSuffix(inName, out _, out _);
@@ -1317,6 +1432,81 @@ namespace PerfSummaries
 			csvFile.Close();
 		}
 
+		// Sorts columns by their diff score. 
+		// Requires max rows. Also works with collated tables with min/max
+		private void SortColumnsByDiffRows()
+		{
+			List<Tuple<SummaryTableColumn, double>> numericColumnSortKeyPairs = new List<Tuple<SummaryTableColumn, double>>();
+			List<SummaryTableColumn> staticColumns = new List<SummaryTableColumn>();
+
+			foreach (SummaryTableColumn column in columns)
+			{
+				if (column.isNumeric && !column.isCountColumn )
+				{
+					double maxDiffScore = column.GetDiffScore();
+					if (hasMinMaxColumns)
+					{
+						// If we have a collated table with avg/min/max then sort by the avg column and we'll re-add the min/max columns later
+						// Columns without a prefix will be treated as static and ordered first
+						if (column.aggregateType == ColumnAggregateType.Avg)
+						{
+							numericColumnSortKeyPairs.Add(new Tuple<SummaryTableColumn, double>(column, maxDiffScore));
+						}
+						else if (column.aggregateType == ColumnAggregateType.None)
+						{
+							staticColumns.Add(column);
+						}
+					}
+					else
+					{
+						numericColumnSortKeyPairs.Add(new Tuple<SummaryTableColumn, double>(column, maxDiffScore));
+					}
+				}
+				else
+				{
+					staticColumns.Add(column);
+				}
+			}
+
+			columns = new List<SummaryTableColumn>();
+			columns.AddRange(staticColumns);
+			numericColumnSortKeyPairs.Sort((a, b) => -a.Item2.CompareTo(b.Item2));
+			foreach (Tuple<SummaryTableColumn, double> pair in numericColumnSortKeyPairs)
+			{
+				SummaryTableColumn column = pair.Item1;
+				columns.Add(column);
+				if (hasMinMaxColumns)
+				{
+					columns.Add(GetAggregateColumn(column, ColumnAggregateType.Min));
+					columns.Add(GetAggregateColumn(column, ColumnAggregateType.Max));
+				}
+			}
+		}
+
+		// Filters out columns with a a max abs diff value below thes specified threshold
+		// Requires max rows. Also works with collated tables with min/max
+		private void FilterColumnsByDiffThreshold(double threshold)
+		{
+			List<SummaryTableColumn> newColumns = new List<SummaryTableColumn>();
+			foreach (SummaryTableColumn column in columns)
+			{
+				if (column.isNumeric && !column.isCountColumn)
+				{
+					// If this is a min or max column then we use the Avg column to compute the diff score. All 3 columns will be filtered together
+					double maxAbsDiff = hasMinMaxColumns ? column.aggregateBaseColumn.GetMaxAbsDiff() : column.GetMaxAbsDiff();
+					if (maxAbsDiff >= threshold)
+					{
+						newColumns.Add(column);
+					}
+				}
+				else
+				{
+					newColumns.Add(column);
+				}
+			}
+			columns = newColumns;
+		}
+
 
 		public void WriteToHTML(
 			string htmlFilename, 
@@ -1363,7 +1553,7 @@ namespace PerfSummaries
 				{
 					for (int i = 0; i < columns.Count; i++)
 					{
-						if (columns[i].name == "Count")
+						if (columns[i].isCountColumn)
 						{
 							stickyColumnCount = i + 1;
 							break;
