When sorting columns by diff, keep CSV stats of the same category together. Do an addional stable sort by stat category after the initial sort by diff

diff --git a/Engine/Source/Programs/CSVTools/PerfReportTool/PerfReportTool.cs b/Engine/Source/Programs/CSVTools/PerfReportTool/PerfReportTool.cs
--- a/Engine/Source/Programs/CSVTools/PerfReportTool/PerfReportTool.cs
+++ b/Engine/Source/Programs/CSVTools/PerfReportTool/PerfReportTool.cs
@@ -26,7 +26,7 @@ namespace PerfReportTool
     class Version
     {
 		// Format: Major.Minor.Bugfix
-        private static string VersionString = "4.105.0";
+        private static string VersionString = "4.106.0";
 
         public static string Get() { return VersionString; }
     };
@@ -169,10 +169,12 @@ namespace PerfReportTool
 			"  -summaryTableXmlRowSortAppend <list,of,stats> : append these stats to the summary table's row sort list\n" +
 			"  -transposeTable : write the summary tables transposed\n" +
 			"  -transposeCollatedTable : write the collated summary table transposed (disables min/max columns)\n" +
-            "  -addDiffRows : adds diff rows after the first two rows\n" +
-			"  -sortColumnsByDiff : sorts columns by the max of its diff scores (use with -addDiffRows)\n" +
-			"      Note: Diff score corresponds to the value of a column's diff row; the sign is reversed if LowIsBad\n" +
-			"  -columnDiffDisplayThreshold <value> : if specified, hides columns with max diff value below this threshold\n" +
+			"Diff rows\n" +
+			"  -addDiffRows : adds diff rows after the first two rows\n" +
+			"  -sortColumnsByDiff : sorts columns by the max of its diff scores (use with -addDiffRows).\n" +
+			"      Notes: Diff score corresponds to the value of a column's diff row; the sign is reversed if LowIsBad\n" +
+			"             Stats are sorted by prefixes first, e.g GPU/\n" +
+			"  -columnDiffDisplayThreshold <value> : if specified, hides columns with max diff value below this threshold\n"+
 			"\n" +
 			"Optional Column Filters\n" +
 			"  -debugShowFilteredColumns : grays out filtered columns instead of removing. Column tooltip will show filtered reason.\n" +
diff --git a/Engine/Source/Programs/CSVTools/PerfReportTool/SummaryTable.cs b/Engine/Source/Programs/CSVTools/PerfReportTool/SummaryTable.cs
--- a/Engine/Source/Programs/CSVTools/PerfReportTool/SummaryTable.cs
+++ b/Engine/Source/Programs/CSVTools/PerfReportTool/SummaryTable.cs
@@ -427,6 +427,20 @@ namespace PerfSummaries
 			return newColumn;
 		}
 
+		// In the case of CSV stats, returns the category. Otherwise, returns an empty string
+		public string GetStatCategory()
+		{
+			if (elementType == SummaryTableElement.Type.CsvStatAverage)
+			{
+				int lastSlashIndex = name.LastIndexOf('/');
+				if (lastSlashIndex != -1)
+				{
+					return name.Substring(0, lastSlashIndex);
+				}
+			}
+			return "";
+		}
+
 		// Returns true if we should displaya min/max/avg sub columns (if enabled).
 		public bool DisplayAggregates()
 		{
@@ -1470,10 +1484,20 @@ namespace PerfSummaries
 
 			columns = new List<SummaryTableColumn>();
 			columns.AddRange(staticColumns);
+
+			// Sort the numeric columns by stat
 			numericColumnSortKeyPairs.Sort((a, b) => -a.Item2.CompareTo(b.Item2));
+			List<SummaryTableColumn> sortedNumericColumns = new List<SummaryTableColumn>();
 			foreach (Tuple<SummaryTableColumn, double> pair in numericColumnSortKeyPairs)
 			{
-				SummaryTableColumn column = pair.Item1;
+				sortedNumericColumns.Add(pair.Item1);
+			}
+
+			// Stable sort the columns by stat prefix
+			IEnumerable<SummaryTableColumn> sortedNumericColumnsOrderedByCategory = sortedNumericColumns.OrderBy(column => column.GetStatCategory());
+
+			foreach (SummaryTableColumn column in sortedNumericColumnsOrderedByCategory)
+			{
 				columns.Add(column);
 				if (hasMinMaxColumns)
 				{
