PerfReportTool - Fix a bug with extra links templating. Some file links were not getting templated due to a change in the ExtraLink class detection logic. Make the detection explicit via a 4th optional link template param in the .links file so we can explicitly specify if a link needs template handling by specifying "linkTemplate" or "noLinkTemplate".

If the new param isn't found or "auto" is specified, detection is automatic (based on the presence of http:// or https://)

diff --git a/Engine/Source/Programs/CSVTools/PerfReportTool/PerfReportTool.cs b/Engine/Source/Programs/CSVTools/PerfReportTool/PerfReportTool.cs
--- a/Engine/Source/Programs/CSVTools/PerfReportTool/PerfReportTool.cs
+++ b/Engine/Source/Programs/CSVTools/PerfReportTool/PerfReportTool.cs
@@ -26,7 +26,7 @@ namespace PerfReportTool
     class Version
     {
 		// Format: Major.Minor.Bugfix
-        private static string VersionString = "4.106.0";
+        private static string VersionString = "4.106.1";
 
         public static string Get() { return VersionString; }
     };
diff --git a/Engine/Source/Programs/CSVTools/PerfReportTool/Summaries/ExtraLinksSummary.cs b/Engine/Source/Programs/CSVTools/PerfReportTool/Summaries/ExtraLinksSummary.cs
--- a/Engine/Source/Programs/CSVTools/PerfReportTool/Summaries/ExtraLinksSummary.cs
+++ b/Engine/Source/Programs/CSVTools/PerfReportTool/Summaries/ExtraLinksSummary.cs
@@ -17,14 +17,34 @@ namespace PerfSummaries
 			public ExtraLink(string fileLine, string inLinkTemplateCsvId)
 			{
 				string[] Sections = fileLine.Split(',');
-				if (Sections.Length != 3)
+				if (Sections.Length < 3)
                 {
 					throw new Exception("Bad links line format: " + fileLine);
 				}
-                ReportText = Sections[0];
+				ReportText = Sections[0];
                 SummaryTableText = Sections[1];
 				LinkURL = Sections[2];
-				if (!LinkURL.StartsWith("http://") && !LinkURL.StartsWith("https://") && !LinkURL.StartsWith("/"))
+				bool bNeedsTemplate = true;
+				string linkTemplateMode = (Sections.Length >= 4) ? Sections[3] : null;
+				if (linkTemplateMode == null || linkTemplateMode == "auto")
+				{
+					// Auto detect if template is needed. If this is a web URL then we assume not
+					if (LinkURL.StartsWith("http://") || LinkURL.StartsWith("https://") )
+					{
+						bNeedsTemplate = false;
+					}
+				}
+				else if (linkTemplateMode == "noLinkTemplate")
+				{
+					bNeedsTemplate = false;
+				}
+				else if (linkTemplateMode != "linkTemplate")
+				{
+					throw new Exception("Bad link type in line: " + fileLine);
+				}
+
+				// If this is a file link then we'll need to template so it can be fixed up externally
+				if (bNeedsTemplate)
 				{
 					LinkTemplateCsvId = inLinkTemplateCsvId;
 				}
