PerfReportTool - Added an "unweighted_average" formula to BoundedStatValueSummaries as weighted averages can sometimes be incorrect.

BoundedStatValueSummaries weights the stat averages by the frametimes. Normally it's not an issue to use the same value for both the value and the weight as they'll end up canceling each other out resulting in an unweighted average. However, in the case where some values are much larger than others this can skew the value resulting in a value much different from the weighted average. This is the case for server frametimes as there can be a few very large spikes that end up skewing the weights.

Since it's not trivial to detect which stats are actually "frametime" stats as they can go by different names, the simplest solution is adding a new "unweighted_average" formula. This also allows using it in other cases where it might make more sense, but these should be very rare.

#rnx

diff --git a/Engine/Source/Programs/CSVTools/PerfReportTool/PerfReportTool.cs b/Engine/Source/Programs/CSVTools/PerfReportTool/PerfReportTool.cs
--- a/Engine/Source/Programs/CSVTools/PerfReportTool/PerfReportTool.cs
+++ b/Engine/Source/Programs/CSVTools/PerfReportTool/PerfReportTool.cs
@@ -26,7 +26,7 @@ namespace PerfReportTool
     class Version
     {
 		// Format: Major.Minor.Bugfix
-        private static string VersionString = "4.106.1";
+        private static string VersionString = "4.107.0";
 
         public static string Get() { return VersionString; }
     };
diff --git a/Engine/Source/Programs/CSVTools/PerfReportTool/Summaries/BoundedStatValuesSummary.cs b/Engine/Source/Programs/CSVTools/PerfReportTool/Summaries/BoundedStatValuesSummary.cs
--- a/Engine/Source/Programs/CSVTools/PerfReportTool/Summaries/BoundedStatValuesSummary.cs
+++ b/Engine/Source/Programs/CSVTools/PerfReportTool/Summaries/BoundedStatValuesSummary.cs
@@ -187,6 +187,17 @@ namespace PerfSummaries
 						}
 					}
 				}
+				else if (col.formula == "unweighted_average")
+				{
+					for (int i = startFrame; i < colEndFrame; i++)
+					{
+						if (col.filterOutZeros == false || statValues[i] > 0)
+						{
+							value += statValues[i];
+							++totalFrameWeight;
+						}
+					}
+				}
 				else if (col.formula == "maximum")
 				{
 					for (int i = startFrame; i < colEndFrame; i++)
