PerfReportTool - Fix a crash when reading partial csvs (just the end so we can read metadata).

#rnx

diff --git a/Engine/Source/Programs/CSVTools/CsvStats/CsvStats.cs b/Engine/Source/Programs/CSVTools/CsvStats/CsvStats.cs
--- a/Engine/Source/Programs/CSVTools/CsvStats/CsvStats.cs
+++ b/Engine/Source/Programs/CSVTools/CsvStats/CsvStats.cs
@@ -1690,7 +1690,7 @@ namespace CSVStats
             }
 
 			// First line is headings, last line contains build info 
-			int numSamples = lines.Count - (bHasMetaData ? 2 : 1);
+			int numSamples = Math.Max(lines.Count - (bHasMetaData ? 2 : 1), 0);
 			if (fileInfo != null)
 			{
 				fileInfo.bIsCsvBin = false;
