CsvToSVG 2.58 (and PerfReportTool 4.107.1) : Fix graphing issue with highlightEventRegion where begin can't be null. This allows us to grey out regions at the start of matches

diff --git a/Engine/Source/Programs/CSVTools/CsvToSvgLib/CsvToSvgLib.cs b/Engine/Source/Programs/CSVTools/CsvToSvgLib/CsvToSvgLib.cs
--- a/Engine/Source/Programs/CSVTools/CsvToSvgLib/CsvToSvgLib.cs
+++ b/Engine/Source/Programs/CSVTools/CsvToSvgLib/CsvToSvgLib.cs
@@ -12,7 +12,7 @@ namespace CSVTools
 
 	public class CsvToSvgLibVersion
 	{
-		private static string VersionString = "2.57";
+		private static string VersionString = "2.58";
 
 		public static string Get() { return VersionString; }
 	};
@@ -1697,10 +1697,14 @@ namespace CSVTools
 						string startName = highlightEventRegions[i * 2].ToLower().Trim();
 						string endName = highlightEventRegions[i * 2 + 1].ToLower().Trim();
 
-						if (endName == "{NULL}")
+						if (endName == "{null}")
 						{
 							endName = null;
 						}
+						if (startName == "{null}")
+						{
+							startName = null;
+						}
 
 						List<int> startIndices = null;
 						List<int> endIndices = null;
diff --git a/Engine/Source/Programs/CSVTools/PerfReportTool/PerfReportTool.cs b/Engine/Source/Programs/CSVTools/PerfReportTool/PerfReportTool.cs
--- a/Engine/Source/Programs/CSVTools/PerfReportTool/PerfReportTool.cs
+++ b/Engine/Source/Programs/CSVTools/PerfReportTool/PerfReportTool.cs
@@ -26,7 +26,7 @@ namespace PerfReportTool
     class Version
     {
 		// Format: Major.Minor.Bugfix
-        private static string VersionString = "4.107.0";
+        private static string VersionString = "4.107.1";
 
         public static string Get() { return VersionString; }
     };
@@ -1640,8 +1640,9 @@ namespace PerfReportTool
 						{
 							highlightEventRegions += ",";
 						}
-						string endEvent = (eventsToStrip[i].endName == null) ? "{NULL}" : eventsToStrip[i].endName;
-						highlightEventRegions += eventsToStrip[i].beginName + "," + endEvent;
+						string endEvent = (eventsToStrip[i].endName == null) ? "{null}" : eventsToStrip[i].endName;
+						string beginEvent = (eventsToStrip[i].beginName == null) ? "{null}" : eventsToStrip[i].beginName;
+						highlightEventRegions += beginEvent + "," + endEvent;
 					}
 					highlightEventRegions += "\"";
 				}
@@ -1779,8 +1780,8 @@ namespace PerfReportTool
 				{
 					for (int i = 0; i < eventsToStrip.Count; i++)
 					{
-						graphParams.highlightEventRegions.Add((eventsToStrip[i].beginName == null) ? "" : eventsToStrip[i].beginName);
-						graphParams.highlightEventRegions.Add((eventsToStrip[i].endName == null) ? "{NULL}" : eventsToStrip[i].endName);
+						graphParams.highlightEventRegions.Add((eventsToStrip[i].beginName == null) ? "{null}" : eventsToStrip[i].beginName);
+						graphParams.highlightEventRegions.Add((eventsToStrip[i].endName == null) ? "{null}" : eventsToStrip[i].endName);
 					}
 				}
 			}
