Fix bug where "minimum" stats that filter out zeros will be zero if their first value is zero

diff --git a/Engine/Source/Programs/CSVTools/PerfReportTool/Summaries/BoundedStatValuesSummary.cs b/Engine/Source/Programs/CSVTools/PerfReportTool/Summaries/BoundedStatValuesSummary.cs
--- a/Engine/Source/Programs/CSVTools/PerfReportTool/Summaries/BoundedStatValuesSummary.cs
+++ b/Engine/Source/Programs/CSVTools/PerfReportTool/Summaries/BoundedStatValuesSummary.cs
@@ -212,8 +212,8 @@ namespace PerfSummaries
 				}
 				else if (col.formula == "minimum")
 				{
-					value = statValues[startFrame];
-					for (int i = startFrame + 1; i < colEndFrame; i++)
+					value = double.MaxValue;
+					for (int i = startFrame; i < colEndFrame; i++)
 					{
 						if (col.filterOutZeros == false || statValues[i] > 0)
 						{
