Added the ability to pass in a list of custom legend entries that match the list of csv files & increased the maxy limit to 8M.

diff --git a/Engine/Source/Programs/CSVTools/CsvToSvgLib/CsvToSvgLib.cs b/Engine/Source/Programs/CSVTools/CsvToSvgLib/CsvToSvgLib.cs
--- a/Engine/Source/Programs/CSVTools/CsvToSvgLib/CsvToSvgLib.cs
+++ b/Engine/Source/Programs/CSVTools/CsvToSvgLib/CsvToSvgLib.cs
@@ -466,6 +466,7 @@ namespace CSVTools
 			// Generate the list of processed, filtered CSV stats for graphing
 			List<CsvStats> csvStatsList = new List<CsvStats>();
 			int currentColorOffset = graphParams.colorOffset;
+			int currentCustomLabelIndex = 0;
 			foreach (CsvInfo csvInfo in csvList)
 			{
 				CsvStats newCsvStats = ProcessCsvStats(csvInfo.stats, graphParams);
@@ -491,7 +492,7 @@ namespace CSVTools
 					stackTotalStatIsAutomatic = true;
 				}
 
-				SetLegend(newCsvStats, csvInfo.filename, graphParams, csvList.Count > 1);
+				SetLegend(newCsvStats, csvInfo.filename, graphParams, csvList.Count > 1, ref currentCustomLabelIndex);
 				if (!graphParams.stacked)
 				{
 					currentColorOffset = AssignColours(newCsvStats, theme, false, currentColorOffset);
@@ -740,7 +741,7 @@ namespace CSVTools
 			}
 
 			// HACK: Clamp to 1m to prevent craziness
-			if (maxSample > 1000000) maxSample = 1000000;
+			if (maxSample > 8000000) maxSample = 8000000;
 
 			if (range.MinY == Range.Auto) newRange.MinY = 0.0f;
 			if (range.MaxY == Range.Auto) newRange.MaxY = maxSample * 1.05f;
@@ -936,9 +937,8 @@ namespace CSVTools
 			return smoothStats;
 		}
 
-		void SetLegend(CsvStats csvStats, string csvFilename, GraphParams graphParams, bool UseFilename)
+		void SetLegend(CsvStats csvStats, string csvFilename, GraphParams graphParams, bool UseFilename, ref int currentCustomLabelIndex)
 		{
-			int customIndex = 0;
 			foreach (StatSamples stat in csvStats.Stats.Values)
 			{
 				stat.LegendName = stat.Name;
@@ -946,10 +946,10 @@ namespace CSVTools
 				{
 					stat.LegendName = System.IO.Path.GetFileName(csvFilename);
 				}
-				if (graphParams.customLegendNames.Count > 0 && customIndex < graphParams.customLegendNames.Count)
+				if (graphParams.customLegendNames.Count > 0 && currentCustomLabelIndex < graphParams.customLegendNames.Count)
 				{
-					stat.LegendName = graphParams.customLegendNames[customIndex];
-					customIndex++;
+					stat.LegendName = graphParams.customLegendNames[currentCustomLabelIndex];
+					currentCustomLabelIndex++;
 				}
 				foreach (string hideStatPrefix in graphParams.hideStatPrefixes)
 				{
