PerfReportTool 4.108.0 + CsvToSvg: -maxAutoMaxY support for graphs + other fixes:
- Add maxAutoMaxY param to graphs and csvToSVG params which allows us to clamp the max automatic Y to a particular value. Also:
- Fix some naming clashes when there are mulltiple graphs on a page. This was preventing us from having graphs with different sizes in a perf report
- Add -embedGraphCommandline arg to help debugging

diff --git a/Engine/Source/Programs/CSVTools/CSVToSVG/CsvToSvgProgram.cs b/Engine/Source/Programs/CSVTools/CSVToSVG/CsvToSvgProgram.cs
--- a/Engine/Source/Programs/CSVTools/CSVToSVG/CsvToSvgProgram.cs
+++ b/Engine/Source/Programs/CSVTools/CSVToSVG/CsvToSvgProgram.cs
@@ -71,6 +71,7 @@ namespace CSVTools
 			"       -legend <list> \n" +
 			"       -maxHierarchyDepth <depth>\n" +
 			"       -minX <value> -maxX <value> -minY <value> -maxY <value>\n" +
+			"       -maxAutoMaxY <value> - clamp automatic maxY to this\n" +
 			"       -noMetadata\n" +
 			"       -noSnap\n" +
 			"       -recurse\n" +
@@ -373,6 +374,7 @@ namespace CSVTools
 			graphParams.maxX = GetFloatArg("maxx", graphParams.maxX);
 			graphParams.minY = GetFloatArg("miny", graphParams.minY);
 			graphParams.maxY = GetFloatArg("maxy", graphParams.maxY);
+			graphParams.maxAutoMaxY = GetFloatArg("maxAutoMaxY", graphParams.maxAutoMaxY);
 
 			graphParams.budget = GetFloatArg("budget", graphParams.budget);
 			graphParams.lineThickness = GetFloatArg("thickness", graphParams.lineThickness);
diff --git a/Engine/Source/Programs/CSVTools/CsvToSvgLib/CsvToSvgLib.cs b/Engine/Source/Programs/CSVTools/CsvToSvgLib/CsvToSvgLib.cs
--- a/Engine/Source/Programs/CSVTools/CsvToSvgLib/CsvToSvgLib.cs
+++ b/Engine/Source/Programs/CSVTools/CsvToSvgLib/CsvToSvgLib.cs
@@ -12,7 +12,7 @@ namespace CSVTools
 
 	public class CsvToSvgLibVersion
 	{
-		private static string VersionString = "2.58";
+		private static string VersionString = "2.59";
 
 		public static string Get() { return VersionString; }
 	};
@@ -67,6 +67,9 @@ namespace CSVTools
 		public float minY = Range.Auto;
 		public float maxY = Range.Auto;
 
+		// Max auto Y range. Set to 0 to disable
+		public float maxAutoMaxY = 0.0f;
+
 		public float budget = float.MaxValue;
 		public float lineThickness = 1.0f;
 
@@ -136,6 +139,11 @@ namespace CSVTools
 			{
 				maxY = Range.Auto;
 			}
+			if (maxAutoMaxY <= 0.0f)
+			{
+				// Clamp to 8m to prevent craziness
+				maxAutoMaxY = 8000000;
+			}
 			if (smoothKernelSize>0.0f)
 			{
 				smoothKernelPercent = 0.0f;
@@ -432,26 +440,26 @@ namespace CSVTools
 
 			// Write defs
 			svg.WriteLine("<defs>");
-			svg.WriteLine("<radialGradient id='radialGradient1'");
+			svg.WriteLine("<radialGradient id='radialGradient<UNIQUE>'");
 			svg.WriteLine("fx='50%' fy='50%' r='65%'");
 			svg.WriteLine("spreadMethod='pad'>");
 			svg.WriteLine("<stop offset='0%'   stop-color=" + theme.BackgroundColourCentre.SVGString() + " stop-opacity='1'/>");
 			svg.WriteLine("<stop offset='100%' stop-color=" + theme.BackgroundColour.SVGString() + " stop-opacity='1' />");
 			svg.WriteLine("</radialGradient>");
 
-			svg.WriteLine("<linearGradient id = 'linearGradient1' x1 = '0%' y1 = '0%' x2 = '100%' y2 = '100%'>");
+			svg.WriteLine("<linearGradient id = 'linearGradient<UNIQUE>' x1 = '0%' y1 = '0%' x2 = '100%' y2 = '100%'>");
 			svg.WriteLine("<stop stop-color = 'black' offset = '0%'/>");
 			svg.WriteLine("<stop stop-color = 'white' offset = '100%'/>");
 			svg.WriteLine("</linearGradient>");
 
 			// Clip rect
-			svg.WriteLine("<clipPath id='graphClipRect'>");
+			svg.WriteLine("<clipPath id='graphClipRect<UNIQUE>'>");
 			// If we're using a fixed point scale for graphs, we need to apply it to the clip rect
 			float clipRectScale = (graphParams.bFixedPointGraphs && !graphParams.percentile) ? graphParams.fixedPointPrecisionScale : 1.0f;
 			svg.WriteLine("<rect  x='" + graphRect.x*clipRectScale + "' y='" + graphRect.y*clipRectScale + "' width='" + graphRect.width*clipRectScale + "' height='" + graphRect.height*clipRectScale + "'/>");
 			svg.WriteLine("</clipPath>");
 
-			svg.WriteLine("<filter id='dropShadowFilter' x='-20%' width='130%' height='130%'>");
+			svg.WriteLine("<filter id='dropShadowFilter<UNIQUE>' x='-20%' width='130%' height='130%'>");
 			svg.WriteLine("<feOffset result='offOut' in='SourceAlpha' dx='-2' dy='2' />");
 			svg.WriteLine("<feGaussianBlur result='blurOut' in='offOut' stdDeviation='1.1' />");
 			svg.WriteLine("<feBlend in='SourceGraphic' in2='blurOut' mode='normal' />");
@@ -560,7 +568,7 @@ namespace CSVTools
 			if (graphParams.stacked) perfLog.LogTiming("StackStats");
 
 			// Adjust range
-			range = ComputeAdjustedYRange(range, graphRect, csvStatsList);
+			range = ComputeAdjustedYRange(range, graphRect, csvStatsList, graphParams.maxAutoMaxY);
 
 			if (graphParams.graphOnly)
 			{
@@ -618,7 +626,7 @@ namespace CSVTools
 
 			// Draw the graphs
 			int csvIndex = 0;
-			svg.WriteLine("<g id='graphArea'>");
+			svg.WriteLine("<g id='graphArea<UNIQUE>'>");
 			if (graphParams.percentile)
 			{
 				range.MinX = graphParams.percentileTop99 ? 99 : (graphParams.percentileTop90 ? 90 : 0);
@@ -716,36 +724,28 @@ namespace CSVTools
 			perfLog.LogTotalTiming(false);
 		}
 
-		Range ComputeAdjustedYRange(Range range, Rect rect, List<CsvStats> csvStats)
+		Range ComputeAdjustedYRange(Range range, Rect rect, List<CsvStats> csvStats, float maxAutoMaxY)
 		{
 			Range newRange = new Range(range.MinX, range.MaxX, range.MinY, range.MaxY);
-			int maxNumSamples = 0;
-			float maxSample = -10000000.0f;
-			float minSample = 10000000.0f;
 
-			foreach (CsvStats stats in csvStats)
+			if (range.MinY == Range.Auto)
 			{
-				foreach (StatSamples samples in stats.Stats.Values)
+				newRange.MinY = 0.0f;
+			}
+
+			if (range.MaxY == Range.Auto)
+			{
+				float maxSample = -10000000.0f;
+				foreach (CsvStats stats in csvStats)
 				{
-					maxNumSamples = Math.Max(maxNumSamples, samples.samples.Count);
-					foreach (float sample in samples.samples)
+					foreach (StatSamples samples in stats.Stats.Values)
 					{
-						if (sample > maxSample)
-						{
-							maxSample = sample;
-						}
-						minSample = Math.Min(minSample, sample);
-						maxSample = Math.Max(maxSample, sample);
+						maxSample = Math.Max(maxSample, samples.ComputeMaxValue());
 					}
 				}
+				newRange.MaxY = Math.Min(maxSample * 1.05f, maxAutoMaxY);
 			}
 
-			// HACK: Clamp to 1m to prevent craziness
-			if (maxSample > 8000000) maxSample = 8000000;
-
-			if (range.MinY == Range.Auto) newRange.MinY = 0.0f;
-			if (range.MaxY == Range.Auto) newRange.MaxY = maxSample * 1.05f;
-
 			// Quantise based on yincrement
 			if (rect != null)
 			{
@@ -1236,7 +1236,7 @@ namespace CSVTools
 
 			if (gradient)
 			{
-				svg.Write("style='fill:url(#radialGradient1);' ");
+				svg.Write("style='fill:url(#radialGradient<UNIQUE>);' ");
 			}
 			if (interactive)
 			{
@@ -1503,7 +1503,7 @@ namespace CSVTools
 							svg.WriteFast(" " + lastX + ","+ originY);
 							svg.WriteFast(" " + firstX + ","+ originY);
 						}
-						svg.WriteLine("' style='fill:" + fillString + ";stroke-width:" + thickness * graphScale + "; clip-path: url(#graphClipRect)' stroke=" + colour.SVGString() + "/>");
+						svg.WriteLine("' style='fill:" + fillString + ";stroke-width:" + thickness * graphScale + "; clip-path: url(#graphClipRect<UNIQUE>)' stroke=" + colour.SVGString() + "/>");
 					}
 				}
 				else
@@ -1533,7 +1533,7 @@ namespace CSVTools
 						svg.WriteFast(" " + ToSvgX((float)n - 20, rect, range) * graphScale + "," + ToSvgY(0.0f, rect, range) * graphScale);
 						svg.WriteFast(" " + ToSvgX(start, rect, range) * graphScale + "," + ToSvgY(0.0f, rect, range) * graphScale );
 					}
-					svg.WriteLine("' style='fill:" + fillString + ";stroke-width:" + thickness * graphScale + "; clip-path: url(#graphClipRect)' stroke=" + colour.SVGString() + "/>");
+					svg.WriteLine("' style='fill:" + fillString + ";stroke-width:" + thickness * graphScale + "; clip-path: url(#graphClipRect<UNIQUE>)' stroke=" + colour.SVGString() + "/>");
 				}
 			}
 		}
@@ -1571,7 +1571,7 @@ namespace CSVTools
 
 				svg.WriteFast(" " + x + "," + y);
 			}
-			svg.WriteLine("' style='fill:none;stroke-width:1.3; clip-path: url(#graphClipRect)' stroke=" + colour.SVGString() + "/>");
+			svg.WriteLine("' style='fill:none;stroke-width:1.3; clip-path: url(#graphClipRect<UNIQUE>)' stroke=" + colour.SVGString() + "/>");
 
 		}
 
diff --git a/Engine/Source/Programs/CSVTools/PerfReportTool/PerfReportTool.cs b/Engine/Source/Programs/CSVTools/PerfReportTool/PerfReportTool.cs
--- a/Engine/Source/Programs/CSVTools/PerfReportTool/PerfReportTool.cs
+++ b/Engine/Source/Programs/CSVTools/PerfReportTool/PerfReportTool.cs
@@ -26,7 +26,7 @@ namespace PerfReportTool
     class Version
     {
 		// Format: Major.Minor.Bugfix
-        private static string VersionString = "4.107.2";
+        private static string VersionString = "4.108.0";
 
         public static string Get() { return VersionString; }
     };
@@ -111,6 +111,7 @@ namespace PerfReportTool
 			"  -csvToSvgSequential : Run CsvToSvg sequentially\n" +
 			"Deprecated performance args:\n" +
 			"  -csvToSvgProcesses : Use separate processes for csvToSVG instead of threads (slower)\n" +
+            "  -embedGraphCommandline : if -csvToSvgProcesses is specified, embeds the commandline for debugging purposes\n" +
 			"  -noBatchedGraphs : disable batched/multithreaded graph generation (use with -csvToSvgProcesses. Default is enabled)\n" +
 			"\n" +
 			"Options to truncate or filter source data:\n" +
@@ -1661,13 +1662,13 @@ namespace PerfReportTool
 				" -maxy " + maxy.ToString() +
 				" -uniqueID Graph_" + graphIndex.ToString() +
 				" -lineDecimalPlaces " + lineDecimalPlaces.ToString() +
-				" -nocommandlineEmbed " +
-
+				( GetBoolArg("embedGraphCommandline") ? "" : " -nocommandlineEmbed") +
 				((statMultiplier != 1.0) ? " -statMultiplier " + statMultiplier.ToString("0.0000000000000000000000") : "") +
 				(hideEventNames ? " -hideeventNames 1" : "") +
 				((minx > 0) ? (" -minx " + minx.ToString()) : "") +
 				((maxx != Int32.MaxValue) ? (" -maxx " + maxx.ToString()) : "") +
 				OptionalHelper.GetDoubleSetting(graphSettings.miny, " -miny ") +
+				OptionalHelper.GetDoubleSetting(graphSettings.maxAutoMaxY, " -maxAutoMaxY ") +
 				OptionalHelper.GetDoubleSetting(graphSettings.threshold, " -threshold ") +
 				OptionalHelper.GetDoubleSetting(graphSettings.averageThreshold, " -averageThreshold ") +
 				OptionalHelper.GetDoubleSetting(minFilterStatValueSetting, " -minFilterStatValue ") +
@@ -1814,6 +1815,10 @@ namespace PerfReportTool
 			}
 			graphParams.maxY = GetFloatArg("maxy", (float)graphSettings.maxy.value);
 
+			if (graphSettings.maxAutoMaxY.isSet)
+			{
+				graphParams.maxAutoMaxY = (float)graphSettings.maxAutoMaxY.value;
+			}
 			if (graphSettings.threshold.isSet)
 			{
 				graphParams.threshold = (float)graphSettings.threshold.value;
diff --git a/Engine/Source/Programs/CSVTools/PerfReportTool/ReportGraph.cs b/Engine/Source/Programs/CSVTools/PerfReportTool/ReportGraph.cs
--- a/Engine/Source/Programs/CSVTools/PerfReportTool/ReportGraph.cs
+++ b/Engine/Source/Programs/CSVTools/PerfReportTool/ReportGraph.cs
@@ -176,6 +176,7 @@ namespace PerfReportTool
 			thickness = new OptionalDouble(element, "thickness");
 			miny = new OptionalDouble(element, "miny");
 			maxy = new OptionalDouble(element, "maxy");
+			maxAutoMaxY = new OptionalDouble(element, "maxAutoMaxY");
 			threshold = new OptionalDouble(element, "threshold");
 			averageThreshold = new OptionalDouble(element, "averageThreshold");
 			minFilterStatValue = new OptionalDouble(element, "minFilterStatValue");
@@ -209,6 +210,7 @@ namespace PerfReportTool
 			thickness.InheritFrom(baseSettings.thickness);
 			miny.InheritFrom(baseSettings.miny);
 			maxy.InheritFrom(baseSettings.maxy);
+			maxAutoMaxY.InheritFrom(baseSettings.maxAutoMaxY);
 			threshold.InheritFrom(baseSettings.threshold);
 			averageThreshold.InheritFrom(baseSettings.averageThreshold);
 			minFilterStatValue.InheritFrom(baseSettings.minFilterStatValue);
@@ -239,6 +241,7 @@ namespace PerfReportTool
 		public OptionalDouble thickness;
 		public OptionalDouble miny;
 		public OptionalDouble maxy;
+		public OptionalDouble maxAutoMaxY;
 		public OptionalDouble threshold;
 		public OptionalDouble averageThreshold;
 		public OptionalDouble minFilterStatValue;
