PerfReportTool - Fix incorrect base directory

Also removed debug directory handling for launching CSVToSVG in debug mode.

#rnx
#tests prt test cases, prs unit tests, local report generation

diff --git a/Engine/Source/Programs/CSVTools/PerfReportTool/PerfReportTool.cs b/Engine/Source/Programs/CSVTools/PerfReportTool/PerfReportTool.cs
--- a/Engine/Source/Programs/CSVTools/PerfReportTool/PerfReportTool.cs
+++ b/Engine/Source/Programs/CSVTools/PerfReportTool/PerfReportTool.cs
@@ -228,18 +228,7 @@ namespace PerfReportTool
 		string GetBaseDirectory()
 		{
 			string location = System.Reflection.Assembly.GetEntryAssembly().Location.ToLower();
-
-			string baseDirectory = location;
-			baseDirectory = baseDirectory.Replace("perfreporttool.exe", "");
-
-			string debugSubDir = "\\bin\\debug\\";
-			if (baseDirectory.ToLower().EndsWith(debugSubDir))
-			{
-				// Might be best to use the CSVToSVG from source, but that might not be built, so use the one checked into binaries instead
-				baseDirectory = baseDirectory.Substring(0, baseDirectory.Length - debugSubDir.Length);
-				baseDirectory += "\\..\\..\\..\\..\\Binaries\\DotNET\\CsvTools";
-			}
-			return baseDirectory;
+			return Path.GetDirectoryName(location);
 		}
 
 		void Run(string[] args)
@@ -1841,7 +1830,7 @@ namespace PerfReportTool
 
 		Process LaunchCsvToSvgAsync(string args)
 		{
-			string csvToolPath = GetBaseDirectory() + "/CSVToSVG.exe";
+			string csvToolPath = Path.Combine(GetBaseDirectory(), "CSVToSVG.exe");
 			string binary = csvToolPath;
 
 			// run mono on non-Windows hosts
diff --git a/Engine/Source/Programs/CSVTools/PerfReportTool/ReportXml.cs b/Engine/Source/Programs/CSVTools/PerfReportTool/ReportXml.cs
--- a/Engine/Source/Programs/CSVTools/PerfReportTool/ReportXml.cs
+++ b/Engine/Source/Programs/CSVTools/PerfReportTool/ReportXml.cs
@@ -78,10 +78,10 @@ namespace PerfReportTool
 		{
 
 			string location = System.Reflection.Assembly.GetEntryAssembly().Location.ToLower();
-			string baseDirectory = location.Replace("perfreporttool.exe", "");
+			string baseDirectory = Path.GetDirectoryName(location);
 
 			// Check if this is a debug build, and redirect base dir to binaries if so
-			if (baseDirectory.Contains("\\engine\\source\\programs\\") && baseDirectory.Contains("\\csvtools\\") && baseDirectory.Contains("\\bin\\debug\\"))
+			if (baseDirectory.Contains("\\engine\\source\\programs\\") && baseDirectory.Contains("\\csvtools\\") && baseDirectory.Contains("\\bin\\debug"))
 			{
 				baseDirectory = baseDirectory.Replace("\\engine\\source\\programs\\", "\\engine\\binaries\\dotnet\\");
 				int csvToolsIndex = baseDirectory.LastIndexOf("\\csvtools\\");
