PerfReportTool 4.210 - Variables part 1 : Add support for metadata references in ReportTypes XML. CSV Metadata references can now be specified Report xml attributes (and values in select cases) with ${meta.KEY} syntax, e.g offsetY="${meta.minimapoffsety}". Multiple variables can be matched in a single attribute if needed, eg offsetXY="${meta.minimapoffsetx},${meta.minimapoffsety}"
Support is provided for Summary classes via XElement extension helper methods, which now take an optional XmlVariableMappings param which handles the variable resolving. GetAttribute, GetSafeAttribute and GetValue support this.
Non metadata variables are also supported in theory e.g ${MyVariable}, but support for populating these does not exist yet
Various cleanup/refactoring and fixups of old code

#tests Passes Test cases with identical results (if no changes to the XML)

diff --git a/Engine/Source/Programs/CSVTools/PerfReportTool/PerfReportTool.cs b/Engine/Source/Programs/CSVTools/PerfReportTool/PerfReportTool.cs
--- a/Engine/Source/Programs/CSVTools/PerfReportTool/PerfReportTool.cs
+++ b/Engine/Source/Programs/CSVTools/PerfReportTool/PerfReportTool.cs
@@ -5,7 +5,6 @@ using System;
 using System.Collections.Generic;
 using System.Linq;
 using System.Text;
-using System.Xml.Linq;
 using System.IO;
 using System.Diagnostics;
 using CSVStats;
@@ -15,18 +14,16 @@ using System.Threading.Tasks;
 using System.Threading;
 
 using PerfSummaries;
-using System.Globalization;
 using CSVTools;
 
 using System.Text.Json;
-using System.Text.Json.Serialization;
 
 namespace PerfReportTool
 {
     class Version
     {
 		// Format: Major.Minor.Bugfix
-        private static string VersionString = "4.200.0";
+        private static string VersionString = "4.210.0";
 
         public static string Get() { return VersionString; }
     };
@@ -1386,7 +1383,7 @@ namespace PerfReportTool
 			if (bExtraLinksSummary)
 			{
 				bool bLinkTemplates = GetBoolArg("linkTemplates");
-				summaries.Insert(0, new ExtraLinksSummary(null, null, bLinkTemplates));
+				summaries.Insert(0, new ExtraLinksSummary(null, reportTypeInfo.vars, null, bLinkTemplates));
 			}
 
 			// If the reporttype has summary info, then write out the summary]
@@ -1638,7 +1635,7 @@ namespace PerfReportTool
 				}
 			}
 
-			OptionalDouble minFilterStatValueSetting = graph.minFilterStatValue.isSet ? graph.minFilterStatValue : graphSettings.minFilterStatValue;
+			Optional<double> minFilterStatValueSetting = graph.minFilterStatValue.isSet ? graph.minFilterStatValue : graphSettings.minFilterStatValue;
 
 			string args =
 				" -csvs \"" + csvFilename + "\"" +
@@ -2015,37 +2012,5 @@ namespace PerfReportTool
 			return null;
 		}
 	}
-
-	static class Extensions
-	{
-		public static T GetSafeAttibute<T>(this XElement element, string attributeName, T defaultValue = default(T))
-		{
-			XAttribute attribute = element.Attribute(attributeName);
-			if (attribute == null)
-			{
-				return defaultValue;
-			}
-
-			try
-			{
-				switch (Type.GetTypeCode(typeof(T)))
-				{
-					case TypeCode.Boolean:
-						return (T)Convert.ChangeType(Convert.ChangeType(attribute.Value, typeof(int)), typeof(bool));
-					case TypeCode.Single:
-					case TypeCode.Double:
-					case TypeCode.Decimal:
-						return (T)Convert.ChangeType(attribute.Value, typeof(T), CultureInfo.InvariantCulture.NumberFormat);
-					default:
-						return (T)Convert.ChangeType(attribute.Value, typeof(T));
-				}
-			}
-			catch (FormatException e)
-			{
-				Console.WriteLine(string.Format("[Warning] Failed to convert XML attribute '{0}' ({1})", attributeName, e.Message));
-				return defaultValue;
-			}
-		}
-    };
 }
 
diff --git a/Engine/Source/Programs/CSVTools/PerfReportTool/ReportGraph.cs b/Engine/Source/Programs/CSVTools/PerfReportTool/ReportGraph.cs
--- a/Engine/Source/Programs/CSVTools/PerfReportTool/ReportGraph.cs
+++ b/Engine/Source/Programs/CSVTools/PerfReportTool/ReportGraph.cs
@@ -17,7 +17,7 @@ namespace PerfReportTool
 		{
 			isSet = false;
 		}
-		public OptionalString(XElement element, string Name, bool IsElement = false)
+		public OptionalString(XElement element, string Name, bool IsElement = false, XmlVariableMappings vars = null )
 		{
 			isSet = false;
 			if (IsElement)
@@ -25,7 +25,7 @@ namespace PerfReportTool
 				XElement child = element.Element(Name);
 				if (child != null)
 				{
-					value = child.Value;
+					value = child.GetValue(vars);
 					isSet = true;
 				}
 			}
@@ -34,7 +34,7 @@ namespace PerfReportTool
 				XAttribute child = element.Attribute(Name);
 				if (child != null)
 				{
-					value = child.Value;
+					value = element.GetRequiredAttribute<string>(vars, Name);
 					isSet = true;
 				}
 			}
@@ -45,100 +45,40 @@ namespace PerfReportTool
 		public string value;
 	};
 
-	class OptionalBool
+	class Optional<T>
 	{
-		public OptionalBool(bool valueIn)
+		public Optional(T valueIn)
 		{
 			value = valueIn;
 			isSet = true;
 		}
-		public OptionalBool()
+		public Optional()
 		{
 			isSet = false;
 		}
-		public OptionalBool(XElement element, string AttributeName)
+		public Optional(XElement element, string AttributeName, XmlVariableMappings vars = null)
 		{
 			isSet = false;
 			try
 			{
 				if (element.Attribute(AttributeName) != null)
 				{
-					value = Convert.ToInt32(element.Attribute(AttributeName).Value) == 1;
+					value = element.GetRequiredAttribute<T>(vars, AttributeName);
 					isSet = true;
 				}
 			}
 			catch { }
 		}
-		public void InheritFrom(OptionalBool baseVersion) { if (!isSet) { isSet = baseVersion.isSet; value = baseVersion.value; } }
+		public void InheritFrom(Optional<T> baseVersion) { if (!isSet) { isSet = baseVersion.isSet; value = baseVersion.value; } }
 
+		public T value;
 		public bool isSet;
-		public bool value;
 	};
 
-	class OptionalInt
-	{
-		public OptionalInt(int valueIn)
-		{
-			value = valueIn;
-			isSet = true;
-		}
-		public OptionalInt()
-		{
-			isSet = false;
-		}
-		public OptionalInt(XElement element, string AttributeName)
-		{
-			isSet = false;
-			try
-			{
-				if (element.Attribute(AttributeName) != null)
-				{
-					value = Convert.ToInt32(element.Attribute(AttributeName).Value);
-					isSet = true;
-				}
-			}
-			catch { }
-		}
-		public void InheritFrom(OptionalInt baseVersion) { if (!isSet) { isSet = baseVersion.isSet; value = baseVersion.value; } }
-
-		public bool isSet;
-		public int value;
-	};
-
-	class OptionalDouble
-	{
-		public OptionalDouble(int valueIn)
-		{
-			value = valueIn;
-			isSet = true;
-		}
-		public OptionalDouble()
-		{
-			isSet = false;
-		}
-		public OptionalDouble(XElement element, string AttributeName)
-		{
-			isSet = false;
-			try
-			{
-				if (element.Attribute(AttributeName) != null)
-				{
-					value = Convert.ToDouble(element.Attribute(AttributeName).Value, System.Globalization.CultureInfo.InvariantCulture);
-					isSet = true;
-				}
-			}
-			catch { }
-		}
-
-		public void InheritFrom(OptionalDouble baseVersion) { if (!isSet) { isSet = baseVersion.isSet; value = baseVersion.value; } }
-
-		public bool isSet;
-		public double value;
-	};
 
 	static class OptionalHelper
 	{
-		public static string GetDoubleSetting(OptionalDouble setting, string cmdline)
+		public static string GetDoubleSetting(Optional<double> setting, string cmdline)
 		{
 			return (setting.isSet ? (cmdline + setting.value.ToString()) : "");
 		}
@@ -151,20 +91,20 @@ namespace PerfReportTool
 
 	class ReportGraph
 	{
-		public ReportGraph(XElement element)
+		public ReportGraph(XElement element, XmlVariableMappings vars)
 		{
-			title = element.Attribute("title").Value;
-			budget = new OptionalDouble(element, "budget");
-			inSummary = element.GetSafeAttibute<bool>("inSummary", false);
-			isExternal = element.GetSafeAttibute<bool>("external", false);
+			title = element.GetRequiredAttribute<string>(vars, "title");
+			budget = new Optional<double>(element, "budget", vars);
+			inSummary = element.GetSafeAttribute<bool>(vars, "inSummary", false);
+			isExternal = element.GetSafeAttribute<bool>(vars, "external", false);
 
-			minFilterStatValue = new OptionalDouble(element, "minFilterStatValue");
+			minFilterStatValue = new Optional<double>(element, "minFilterStatValue", vars);
 		}
 		public string title;
-		public OptionalDouble budget;
+		public Optional<double> budget;
 		public bool inSummary;
 		public bool isExternal;
-		public OptionalDouble minFilterStatValue;
+		public Optional<double> minFilterStatValue;
 		public GraphSettings settings;
 	};
 
@@ -172,36 +112,36 @@ namespace PerfReportTool
 	{
 		public GraphSettings(XElement element)
 		{
-			smooth = new OptionalBool(element, "smooth");
-			thickness = new OptionalDouble(element, "thickness");
-			miny = new OptionalDouble(element, "miny");
-			maxy = new OptionalDouble(element, "maxy");
-			maxAutoMaxY = new OptionalDouble(element, "maxAutoMaxY");
-			threshold = new OptionalDouble(element, "threshold");
-			averageThreshold = new OptionalDouble(element, "averageThreshold");
-			minFilterStatValue = new OptionalDouble(element, "minFilterStatValue");
+			smooth = new Optional<bool>(element, "smooth");
+			thickness = new Optional<double>(element, "thickness");
+			miny = new Optional<double>(element, "miny");
+			maxy = new Optional<double>(element, "maxy");
+			maxAutoMaxY = new Optional<double>(element, "maxAutoMaxY");
+			threshold = new Optional<double>(element, "threshold");
+			averageThreshold = new Optional<double>(element, "averageThreshold");
+			minFilterStatValue = new Optional<double>(element, "minFilterStatValue");
 			minFilterStatName = new OptionalString(element, "minFilterStatName");
-			smoothKernelPercent = new OptionalDouble(element, "smoothKernelPercent");
-			smoothKernelSize = new OptionalDouble(element, "smoothKernelSize");
-			compression = new OptionalDouble(element, "compression");
-			width = new OptionalInt(element, "width");
-			height = new OptionalInt(element, "height");
-			stacked = new OptionalBool(element, "stacked");
-			showAverages = new OptionalBool(element, "showAverages");
-			filterOutZeros = new OptionalBool(element, "filterOutZeros");
-			maxHierarchyDepth = new OptionalInt(element, "maxHierarchyDepth");
+			smoothKernelPercent = new Optional<double>(element, "smoothKernelPercent");
+			smoothKernelSize = new Optional<double>(element, "smoothKernelSize");
+			compression = new Optional<double>(element, "compression");
+			width = new Optional<int>(element, "width");
+			height = new Optional<int>(element, "height");
+			stacked = new Optional<bool>(element, "stacked");
+			showAverages = new Optional<bool>(element, "showAverages");
+			filterOutZeros = new Optional<bool>(element, "filterOutZeros");
+			maxHierarchyDepth = new Optional<int>(element, "maxHierarchyDepth");
 			hideStatPrefix = new OptionalString(element, "hideStatPrefix");
 			mainStat = new OptionalString(element, "mainStat");
 			showEvents = new OptionalString(element, "showEvents");
-			requiresDetailedStats = new OptionalBool(element, "requiresDetailedStats");
+			requiresDetailedStats = new Optional<bool>(element, "requiresDetailedStats");
 			ignoreStats = new OptionalString(element, "ignoreStats");
 
 			statString = new OptionalString(element, "statString", true);
 			//additionalArgs = new OptionalString(element, "additionalArgs", true);
-			statMultiplier = new OptionalDouble(element, "statMultiplier");
-			legendAverageThreshold = new OptionalDouble(element, "legendAverageThreshold");
-			snapToPeaks = new OptionalBool(element, "snapToPeaks");
-			lineDecimalPlaces = new OptionalInt(element, "lineDecimalPlaces");
+			statMultiplier = new	(element, "statMultiplier");
+			legendAverageThreshold = new Optional<double>(element, "legendAverageThreshold");
+			snapToPeaks = new Optional<bool>(element, "snapToPeaks");
+			lineDecimalPlaces = new Optional<int>(element, "lineDecimalPlaces");
 		}
 		public void InheritFrom(GraphSettings baseSettings)
 		{
@@ -236,36 +176,36 @@ namespace PerfReportTool
 			lineDecimalPlaces.InheritFrom(baseSettings.lineDecimalPlaces);
 
 		}
-		public OptionalBool smooth;
+		public Optional<bool> smooth;
 		public OptionalString statString;
-		public OptionalDouble thickness;
-		public OptionalDouble miny;
-		public OptionalDouble maxy;
-		public OptionalDouble maxAutoMaxY;
-		public OptionalDouble threshold;
-		public OptionalDouble averageThreshold;
-		public OptionalDouble minFilterStatValue;
+		public Optional<double> thickness;
+		public Optional<double> miny;
+		public Optional<double> maxy;
+		public Optional<double> maxAutoMaxY;
+		public Optional<double> threshold;
+		public Optional<double> averageThreshold;
+		public Optional<double> minFilterStatValue;
 		public OptionalString minFilterStatName;
-		public OptionalDouble smoothKernelSize;
-		public OptionalDouble smoothKernelPercent;
-		public OptionalDouble compression;
-		public OptionalInt width;
-		public OptionalInt height;
+		public Optional<double> smoothKernelSize;
+		public Optional<double> smoothKernelPercent;
+		public Optional<double> compression;
+		public Optional<int> width;
+		public Optional<int> height;
 		//public OptionalString additionalArgs;
-		public OptionalBool stacked;
-		public OptionalBool showAverages;
-		public OptionalBool filterOutZeros;
-		public OptionalInt maxHierarchyDepth;
+		public Optional<bool> stacked;
+		public Optional<bool> showAverages;
+		public Optional<bool> filterOutZeros;
+		public Optional<int> maxHierarchyDepth;
 		public OptionalString hideStatPrefix;
 		public OptionalString mainStat;
 		public OptionalString showEvents;
 		public OptionalString ignoreStats;
-		public OptionalDouble statMultiplier;
-		public OptionalDouble legendAverageThreshold;
+		public Optional<double> statMultiplier;
+		public Optional<double> legendAverageThreshold;
 
-		public OptionalBool requiresDetailedStats;
-		public OptionalBool snapToPeaks;
-		public OptionalInt lineDecimalPlaces;
+		public Optional<bool> requiresDetailedStats;
+		public Optional<bool> snapToPeaks;
+		public Optional<int> lineDecimalPlaces;
 
 	};
 
diff --git a/Engine/Source/Programs/CSVTools/PerfReportTool/ReportType.cs b/Engine/Source/Programs/CSVTools/PerfReportTool/ReportType.cs
--- a/Engine/Source/Programs/CSVTools/PerfReportTool/ReportType.cs
+++ b/Engine/Source/Programs/CSVTools/PerfReportTool/ReportType.cs
@@ -2,8 +2,10 @@
 
 using System;
 using System.Collections.Generic;
+using System.Runtime.InteropServices;
 using System.Text;
 using System.Xml.Linq;
+using CSVStats;
 using PerfSummaries;
 
 namespace PerfReportTool
@@ -16,19 +18,22 @@ namespace PerfReportTool
 
 	class ReportTypeInfo
 	{
-		public ReportTypeInfo(XElement element, Dictionary<string, XElement> sharedSummaries, string baseXmlDirectory)
+		public ReportTypeInfo(XElement element, Dictionary<string, XElement> sharedSummaries, string baseXmlDirectory, XmlVariableMappings inVariableMappings)
 		{
 			graphs = new List<ReportGraph>();
 			summaries = new List<Summary>();
-			name = element.Attribute("name").Value;
-			title = element.Attribute("title").Value;
-			bStripEvents = element.GetSafeAttibute<bool>("stripEvents", true);
+
+			vars = inVariableMappings;
+
+			name = element.GetRequiredAttribute<string>(vars, "name");
+			title = element.GetRequiredAttribute<string>(vars, "title");
+			bStripEvents = element.GetSafeAttribute<bool>(vars, "stripEvents", true);
 
 			foreach (XElement child in element.Elements())
 			{
 				if (child.Name == "graph")
 				{
-					ReportGraph graph = new ReportGraph(child);
+					ReportGraph graph = new ReportGraph(child, vars);
 					graphs.Add(graph);
 				}
 				else if (child.Name == "summary" || child.Name == "summaryRef")
@@ -36,14 +41,14 @@ namespace PerfReportTool
 					XElement summaryElement = null;
 					if (child.Name == "summaryRef")
 					{
-						summaryElement = sharedSummaries[child.Attribute("name").Value];
+						summaryElement = sharedSummaries[child.GetRequiredAttribute<string>(vars, "name")];
 					}
 					else
 					{
 						summaryElement = child;
 					}
-					string summaryType = summaryElement.Attribute("type").Value;
-					summaries.Add(SummaryFactory.Create(summaryType, summaryElement, baseXmlDirectory));
+					string summaryType = summaryElement.GetRequiredAttribute<string>(vars, "type");
+					summaries.Add(SummaryFactory.Create(summaryType, summaryElement, vars, baseXmlDirectory));
 				}
 				else if (child.Name == "metadataToShow")
 				{
@@ -82,5 +87,7 @@ namespace PerfReportTool
 		public string[] metadataToShowList;
 		public string summaryTableCacheID;
 		public bool bStripEvents;
+
+		public XmlVariableMappings vars;
 	};
 }
\ No newline at end of file
diff --git a/Engine/Source/Programs/CSVTools/PerfReportTool/ReportXml.cs b/Engine/Source/Programs/CSVTools/PerfReportTool/ReportXml.cs
--- a/Engine/Source/Programs/CSVTools/PerfReportTool/ReportXml.cs
+++ b/Engine/Source/Programs/CSVTools/PerfReportTool/ReportXml.cs
@@ -7,6 +7,8 @@ using System.Linq;
 using System.Xml.Linq;
 using CSVStats;
 using PerfSummaries;
+using System.Globalization;
+using System.Runtime.InteropServices;
 
 namespace PerfReportTool
 {
@@ -76,7 +78,6 @@ namespace PerfReportTool
 			string summaryTableXmlAppendStr,
 			string summaryTableXmlRowSortAppendStr )
 		{
-
 			string location = System.Reflection.Assembly.GetEntryAssembly().Location.ToLower();
 			string baseDirectory = Path.GetDirectoryName(location);
 
@@ -154,7 +155,7 @@ namespace PerfReportTool
 			}
 			else
 			{
-				graphsXMLFilename = reportTypesElement.GetSafeAttibute<string>("reportGraphsFile");
+				graphsXMLFilename = reportTypesElement.GetSafeAttribute<string>("reportGraphsFile");
 				if (graphsXMLFilename != null)
 				{
 					graphsXMLFilename = Path.GetDirectoryName(reportTypeXmlFilename) + "\\" + graphsXMLFilename;
@@ -165,7 +166,7 @@ namespace PerfReportTool
 				}
 
 			}
-			defaultReportTypeName = reportTypesElement.GetSafeAttibute<string>("default");
+			defaultReportTypeName = reportTypesElement.GetSafeAttribute<string>("default");
 
 			Console.Out.WriteLine("GraphXML:  " + graphsXMLFilename+"\n");
 			XDocument reportGraphsDoc = XDocument.Load(graphsXMLFilename);
@@ -191,7 +192,7 @@ namespace PerfReportTool
 					{
 						if (graphElement.Name == "graph")
 						{
-							string title = graphElement.Attribute("title").Value.ToLower();
+							string title = graphElement.GetRequiredAttribute<string>("title").ToLower();
 							GraphSettings graphSettings = new GraphSettings(graphElement);
 							graphSettings.InheritFrom(groupSettings);
 							graphs.Add(title, graphSettings);
@@ -209,8 +210,8 @@ namespace PerfReportTool
 			{
 				foreach (XElement mapping in displayNameElement.Elements("mapping"))
 				{
-					string statName = mapping.GetSafeAttibute<string>("statName");
-					string displayName = mapping.GetSafeAttibute<string>("displayName");
+					string statName = mapping.GetSafeAttribute<string>("statName");
+					string displayName = mapping.GetSafeAttribute<string>("displayName");
 					if (statName != null && displayName != null)
 					{
 						statDisplayNameMapping.Add(statName.ToLower(), displayName);
@@ -231,10 +232,10 @@ namespace PerfReportTool
 			{
 				foreach (XElement mapping in derivedMetadataMappingsElement.Elements("mapping"))
 				{
-					string sourceName = mapping.GetSafeAttibute<string>("sourceName");
-					string sourceValue = mapping.GetSafeAttibute<string>("sourceValue");
-					string destName = mapping.GetSafeAttibute<string>("destName");
-					string destValue = mapping.GetSafeAttibute<string>("destValue");
+					string sourceName = mapping.GetSafeAttribute<string>("sourceName");
+					string sourceValue = mapping.GetSafeAttribute<string>("sourceValue");
+					string destName = mapping.GetSafeAttribute<string>("destName");
+					string destValue = mapping.GetSafeAttribute<string>("destValue");
 					if (sourceName == null || sourceValue == null || destName == null || destValue == null)
 					{
 						throw new Exception("Derivedmetadata mapping is missing a required attribute!\nRequired attributes: sourceName, sourceValue, destName, destValue.\nXML: " + mapping.ToString());
@@ -251,8 +252,8 @@ namespace PerfReportTool
 				foreach (XElement eventPair in eventsToStripEl.Elements("eventPair"))
 				{
 					CsvEventStripInfo eventInfo = new CsvEventStripInfo();
-					eventInfo.beginName = eventPair.GetSafeAttibute<string>("begin");
-					eventInfo.endName = eventPair.GetSafeAttibute<string>("end");
+					eventInfo.beginName = eventPair.GetSafeAttribute<string>("begin");
+					eventInfo.endName = eventPair.GetSafeAttribute<string>("end");
 
 					if (eventInfo.beginName == null && eventInfo.endName == null)
 					{
@@ -298,7 +299,7 @@ namespace PerfReportTool
 				foreach (XElement summaryElement in summaryTablesElement.Elements("summaryTable"))
 				{
 					SummaryTableInfo table = new SummaryTableInfo(summaryElement, substitutionsDict, appendList, rowSortAppendList);
-					summaryTables.Add(summaryElement.Attribute("name").Value.ToLower(), table);
+					summaryTables.Add(summaryElement.GetRequiredAttribute<string>("name").ToLower(), table);
 				}
 			}
 
@@ -309,7 +310,7 @@ namespace PerfReportTool
 			{
 				foreach (XElement summaryElement in sharedSummariesElement.Elements("summary"))
 				{
-					sharedSummaries.Add(summaryElement.Attribute("refName").Value, summaryElement);
+					sharedSummaries.Add(summaryElement.GetRequiredAttribute<string>("refName"), summaryElement);
 				}
 			}
 
@@ -317,6 +318,17 @@ namespace PerfReportTool
 
 		public ReportTypeInfo GetReportTypeInfo(string reportType, CachedCsvFile csvFile, bool bBulkMode, bool forceReportType)
 		{
+			// Setup the variable mappings
+			XmlVariableMappings vars = new XmlVariableMappings();
+			if (csvFile.metadata != null)
+			{
+				Dictionary<string, string> metadataDict = csvFile.metadata.Values;
+				foreach (string key in metadataDict.Keys)
+				{
+					vars.SetVariable("meta." + key, metadataDict[key]);
+				}
+			}
+
 			ReportTypeInfo reportTypeInfo = null;
 			if (reportType == "")
 			{
@@ -324,15 +336,15 @@ namespace PerfReportTool
 				// Attempt to determine the report type automatically based on the stats
 				foreach (XElement element in reportTypesElement.Elements("reporttype"))
 				{
-					bool bReportTypeSupportsAutodetect = element.GetSafeAttibute<bool>("allowAutoDetect", true);
+					bool bReportTypeSupportsAutodetect = element.GetSafeAttribute<bool>(vars, "allowAutoDetect", true);
 
-					if (bReportTypeSupportsAutodetect && IsReportTypeXMLCompatibleWithStats(element, csvFile.dummyCsvStats))
+					if (bReportTypeSupportsAutodetect && IsReportTypeXMLCompatibleWithStats(element, csvFile.dummyCsvStats, vars))
 					{
-						reportTypeInfo = new ReportTypeInfo(element, sharedSummaries, baseXmlDirectory);
+						reportTypeInfo = new ReportTypeInfo(element, sharedSummaries, baseXmlDirectory, vars);
 						break;
 					}
 
-					if (defaultReportTypeName != null && element.GetSafeAttibute<string>("name") == defaultReportTypeName )
+					if (defaultReportTypeName != null && element.GetSafeAttribute<string>(vars, "name") == defaultReportTypeName )
 					{
 						defaultReportTypeElement = element;
 					}
@@ -344,12 +356,12 @@ namespace PerfReportTool
 					{
 						throw new Exception("Default report type " + defaultReportTypeName + " was not found in " + reportTypeXmlFilename);
 					}
-					if (!IsReportTypeXMLCompatibleWithStats(defaultReportTypeElement, csvFile.dummyCsvStats, true))
+					if (!IsReportTypeXMLCompatibleWithStats(defaultReportTypeElement, csvFile.dummyCsvStats, vars, true))
 					{
 						throw new Exception("Default report type " + defaultReportTypeName + " was not compatible with CSV " + csvFile.filename);
 					}
 					Console.Out.WriteLine("Falling back to default report type: " + defaultReportTypeName);
-					reportTypeInfo = new ReportTypeInfo(defaultReportTypeElement, sharedSummaries, baseXmlDirectory);
+					reportTypeInfo = new ReportTypeInfo(defaultReportTypeElement, sharedSummaries, baseXmlDirectory, vars);
 				}
 				else if (reportTypeInfo == null)
 				{
@@ -361,7 +373,7 @@ namespace PerfReportTool
 				XElement foundReportTypeElement = null;
 				foreach (XElement element in reportTypesElement.Elements("reporttype"))
 				{
-					if (element.Attribute("name").Value.ToLower() == reportType)
+					if (element.GetSafeAttribute<string>(vars, "name").ToLower() == reportType)
 					{
 						foundReportTypeElement = element;
 					}
@@ -371,7 +383,7 @@ namespace PerfReportTool
 					throw new Exception("Report type " + reportType + " not found in " + reportTypeXmlFilename);
 				}
 
-				if (!IsReportTypeXMLCompatibleWithStats(foundReportTypeElement, csvFile.dummyCsvStats))
+				if (!IsReportTypeXMLCompatibleWithStats(foundReportTypeElement, csvFile.dummyCsvStats, vars))
 				{
 					if (forceReportType)
 					{
@@ -382,7 +394,7 @@ namespace PerfReportTool
 						throw new Exception("Report type " + reportType + " is not compatible with CSV " + csvFile.filename);
 					}
 				}
-				reportTypeInfo = new ReportTypeInfo(foundReportTypeElement, sharedSummaries, baseXmlDirectory);
+				reportTypeInfo = new ReportTypeInfo(foundReportTypeElement, sharedSummaries, baseXmlDirectory, vars);
 			}
 
 			// Load the graphs
@@ -406,24 +418,24 @@ namespace PerfReportTool
 			return reportTypeInfo;
 		}
 
-		bool IsReportTypeXMLCompatibleWithStats(XElement reportTypeElement, CsvStats csvStats, bool bIsDefaultFallback=false)
+		bool IsReportTypeXMLCompatibleWithStats(XElement reportTypeElement, CsvStats csvStats, XmlVariableMappings vars, bool bIsDefaultFallback=false)
 		{
-			XAttribute nameAt = reportTypeElement.Attribute("name");
-			if (nameAt == null)
+			string name = reportTypeElement.GetSafeAttribute<string>(vars, "name");
+			if (name == null)
 			{
 				return false;
 			}
-			string reportTypeName = nameAt.Value;
+			string reportTypeName = name;
 
 			XElement autoDetectionEl = reportTypeElement.Element("autodetection");
 			if (autoDetectionEl == null)
 			{
 				return false;
 			}
-			XAttribute requiredStatsAt = autoDetectionEl.Attribute("requiredstats");
-			if (requiredStatsAt != null)
+			string requiredStatsStr = autoDetectionEl.GetSafeAttribute<string>(vars, "requiredstats");
+			if (requiredStatsStr != null)
 			{
-				string[] requiredStats = requiredStatsAt.Value.Split(',');
+				string[] requiredStats = requiredStatsStr.Split(',');
 				foreach (string stat in requiredStats)
 				{
 					if (csvStats.GetStatsMatchingString(stat).Count == 0)
@@ -435,12 +447,12 @@ namespace PerfReportTool
 
 			foreach (XElement requiredMetadataEl in autoDetectionEl.Elements("requiredmetadata"))
 			{
-				XAttribute keyAt = requiredMetadataEl.Attribute("key");
-				if (keyAt == null)
+				string key = requiredMetadataEl.GetSafeAttribute<string>(vars, "key");
+				if (key == null)
 				{
 					throw new Exception("Report type " + reportTypeName + " has no 'key' attribute!");
 				}
-				XAttribute allowedValuesAt = requiredMetadataEl.Attribute("allowedValues");
+				string allowedValuesAt = requiredMetadataEl.GetSafeAttribute<string>(vars, "allowedValues");
 				if (allowedValuesAt == null)
 				{
 					throw new Exception("Report type " + reportTypeName + " has no 'allowedValues' attribute!");
@@ -453,19 +465,19 @@ namespace PerfReportTool
 				}
 
 				// Some metadata may be safe to skip for the default fallback case
-				if (bIsDefaultFallback && requiredMetadataEl.GetSafeAttibute("ignoreForDefaultFallback", false))
+				if (bIsDefaultFallback && requiredMetadataEl.GetSafeAttribute(vars, "ignoreForDefaultFallback", false))
 				{
 					continue;
 				}
 
-				bool ignoreIfKeyNotFound = requiredMetadataEl.GetSafeAttibute("ignoreIfKeyNotFound", true);
-				bool stopIfKeyFound = requiredMetadataEl.GetSafeAttibute("stopIfKeyFound", false);
+				bool ignoreIfKeyNotFound = requiredMetadataEl.GetSafeAttribute(vars, "ignoreIfKeyNotFound", true);
+				bool stopIfKeyFound = requiredMetadataEl.GetSafeAttribute(vars, "stopIfKeyFound", false);
 
-				string key = keyAt.Value.ToLower();
+				key = key.ToLower();
 				if (csvStats.metaData.Values.ContainsKey(key))
 				{
 					string value = csvStats.metaData.Values[key].ToLower();
-					string[] allowedValues = allowedValuesAt.Value.ToString().ToLower().Split(',');
+					string[] allowedValues = allowedValuesAt.ToString().ToLower().Split(',');
 					if (!allowedValues.Contains(value))
 					{
 						return false;
@@ -533,4 +545,134 @@ namespace PerfReportTool
 		string reportTypeXmlFilename;
 		public DerivedMetadataMappings derivedMetadataMappings;
 	}
+
+	class XmlVariableMappings
+	{
+		public void SetVariable(string Name, string Value)
+		{
+			vars[Name] = Value;
+		}
+
+		public string ResolveVariables(string attributeValue)
+		{
+			// Remap all variables found in the attribute name
+			if (!attributeValue.Contains('$'))
+			{
+				return attributeValue;
+			}
+
+			// Remap all variables found in the attribute name
+			int StringPos = 0;
+			while (StringPos < attributeValue.Length)
+			{
+				int VarStartIndex = attributeValue.IndexOf("${", StringPos);
+				if (VarStartIndex == -1)
+				{
+					break;
+				}
+				int VarEndIndex = attributeValue.IndexOf('}', VarStartIndex+1);
+				if (VarEndIndex == -1)
+				{
+					break;
+				}
+				string VariableName = attributeValue.Substring(VarStartIndex+2, VarEndIndex - VarStartIndex-2);
+
+				// Replace the variable if found
+				if (vars.TryGetValue(VariableName, out string VariableValue))
+				{
+					attributeValue = attributeValue.Substring(0,VarStartIndex) + VariableValue + attributeValue.Substring(VarEndIndex+1);
+					StringPos = VarStartIndex + VariableValue.Length;
+				}
+				else
+				{
+					Console.WriteLine("[Warning] Failed to resolve variable $(" + VariableName + "}");
+					StringPos = VarEndIndex;
+				}
+			}
+			return attributeValue;
+		}
+		Dictionary<string, string> vars = new Dictionary<string, string>();
+	}
+	
+
+	static class Extensions
+	{
+		public static string GetValue(this XElement element, XmlVariableMappings xmlVariableMappings = null)
+		{
+			string value = element.Value;
+			if (xmlVariableMappings != null)
+			{
+				value = xmlVariableMappings.ResolveVariables(value);
+			}
+			return value;
+		}
+		public static T GetRequiredAttribute<T>(this XElement element, string attributeName)
+		{
+			return GetAttributeInternal<T>(element, null, attributeName, default, true);
+		}
+
+		public static T GetRequiredAttribute<T>(this XElement element, XmlVariableMappings xmlVariableMappings, string attributeName)
+		{
+			return GetAttributeInternal<T>(element, xmlVariableMappings, attributeName, default, true );
+		}
+
+		public static T GetSafeAttribute<T>(this XElement element, string attributeName, T defaultValue = default)
+		{
+			return GetAttributeInternal<T>(element, null,  attributeName, defaultValue, false);
+		}
+
+		public static T GetSafeAttribute<T>(this XElement element, XmlVariableMappings xmlVariableMappings, string attributeName, T defaultValue = default)
+		{
+			return GetAttributeInternal<T>(element, xmlVariableMappings, attributeName, defaultValue, false);
+		}
+
+		private static T GetAttributeInternal<T>(this XElement element, XmlVariableMappings xmlVariableMappings, string attributeName, T defaultValue, bool throwIfNotFound)
+		{
+			XAttribute attribute = element.Attribute(attributeName);
+			if (attribute == null)
+			{
+				if (throwIfNotFound)
+				{
+					throw new Exception("Attribute "+attributeName+" not found in element "+element.Name);
+				}
+				return defaultValue;
+			}
+
+			// Resolve variables if a mapping is provided
+			string attributeValue = attribute.Value;
+			if (xmlVariableMappings != null)
+			{
+				attributeValue = xmlVariableMappings.ResolveVariables(attributeValue);
+			}
+
+			try
+			{
+				switch (Type.GetTypeCode(typeof(T)))
+				{
+					case TypeCode.Boolean:
+						try
+						{
+							// Support int/bool conversion
+							return (T)Convert.ChangeType(Convert.ChangeType(attributeValue, typeof(int)), typeof(bool));
+						}
+						catch (FormatException)
+						{
+							// fall back to reading it as an actual bool
+							return (T)Convert.ChangeType(attributeValue, typeof(T));
+						}
+					case TypeCode.Single:
+					case TypeCode.Double:
+					case TypeCode.Decimal:
+						return (T)Convert.ChangeType(attributeValue, typeof(T), CultureInfo.InvariantCulture.NumberFormat);
+					default:
+						return (T)Convert.ChangeType(attributeValue, typeof(T));
+				}
+			}
+			catch (FormatException e)
+			{
+				Console.WriteLine(string.Format("[Warning] Failed to convert XML attribute '{0}' '{1}' ({2})", attributeName, attributeValue, e.Message));
+				return defaultValue;
+			}
+		}
+	};
 }
\ No newline at end of file
diff --git a/Engine/Source/Programs/CSVTools/PerfReportTool/Summaries/BoundedStatValuesSummary.cs b/Engine/Source/Programs/CSVTools/PerfReportTool/Summaries/BoundedStatValuesSummary.cs
--- a/Engine/Source/Programs/CSVTools/PerfReportTool/Summaries/BoundedStatValuesSummary.cs
+++ b/Engine/Source/Programs/CSVTools/PerfReportTool/Summaries/BoundedStatValuesSummary.cs
@@ -28,30 +28,25 @@ namespace PerfSummaries
 			public double statExponent; // Exponent for stat value in streamingstressmetric formula
 			public ColourThresholdList colourThresholdList;
 		};
-		public BoundedStatValuesSummary(XElement element, string baseXmlDirectory)
+		public BoundedStatValuesSummary(XElement element, XmlVariableMappings vars, string baseXmlDirectory)
 		{
-			ReadStatsFromXML(element);
+			ReadStatsFromXML(element, vars);
 			if (stats.Count != 0)
 			{
 				throw new Exception("<stats> element is not supported");
 			}
 
-			title = element.GetSafeAttibute("title", "Events");
-			beginEvent = element.GetSafeAttibute<string>("beginevent");
-			endEvent = element.GetSafeAttibute<string>("endevent");
+			title = element.GetSafeAttribute(vars, "title", "Events");
+			beginEvent = element.GetSafeAttribute<string>(vars, "beginevent");
+			endEvent = element.GetSafeAttribute<string>(vars, "endevent");
 
-			endOffsetPercentage = 0.0;
-			XAttribute endOffsetAtt = element.Attribute("endoffsetpercent");
-			if (endOffsetAtt != null)
-			{
-				endOffsetPercentage = double.Parse(endOffsetAtt.Value);
-			}
+			endOffsetPercentage = element.GetSafeAttribute<double>(vars, "endoffsetpercent", 0.0);
 			columns = new List<Column>();
 
 			foreach (XElement columnEl in element.Elements("column"))
 			{
 				Column column = new Column();
-				double[] colourThresholds = ReadColourThresholdsXML(columnEl.Element("colourThresholds"));
+				double[] colourThresholds = ReadColourThresholdsXML(columnEl.Element("colourThresholds"), vars);
 				if (colourThresholds != null)
 				{
 					column.colourThresholdList = new ColourThresholdList();
@@ -61,27 +56,23 @@ namespace PerfSummaries
 					}
 				}
 
-				XAttribute summaryStatNameAtt = columnEl.Attribute("summaryStatName");
-				if (summaryStatNameAtt != null)
-				{
-					column.summaryStatName = summaryStatNameAtt.Value;
-				}
-				column.statName = columnEl.Attribute("stat").Value.ToLower();
+				column.summaryStatName = columnEl.GetSafeAttribute<string>(vars, "summaryStatName");
+				column.statName = columnEl.GetRequiredAttribute<string>(vars, "stat").ToLower();
 				if (!stats.Contains(column.statName))
 				{
 					stats.Add(column.statName);
 				}
-				column.otherStatName = columnEl.GetSafeAttibute<string>("otherStat", "").ToLower();
+				column.otherStatName = columnEl.GetSafeAttribute<string>("otherStat", "").ToLower();
 
-				column.name = columnEl.Attribute("name").Value;
-				column.formula = columnEl.Attribute("formula").Value.ToLower();
-				column.filterOutZeros = columnEl.GetSafeAttibute<bool>("filteroutzeros", false);
-				column.perSecond = columnEl.GetSafeAttibute<bool>("persecond", false);
-				column.multiplier = columnEl.GetSafeAttibute<double>("multiplier", 1.0);
-				column.threshold = columnEl.GetSafeAttibute<double>("threshold", 0.0);
-				column.applyEndOffset = columnEl.GetSafeAttibute<bool>("applyEndOffset", true);
-				column.frameExponent = columnEl.GetSafeAttibute<double>("frameExponent", 4.0);
-				column.statExponent = columnEl.GetSafeAttibute<double>("statExponent", 0.25);
+				column.name = columnEl.GetRequiredAttribute<string>(vars, "name");
+				column.formula = columnEl.GetRequiredAttribute<string>(vars, "formula").ToLower();
+				column.filterOutZeros = columnEl.GetSafeAttribute<bool>(vars, "filteroutzeros", false);
+				column.perSecond = columnEl.GetSafeAttribute<bool>(vars, "persecond", false);
+				column.multiplier = columnEl.GetSafeAttribute<double>(vars, "multiplier", 1.0);
+				column.threshold = columnEl.GetSafeAttribute<double>(vars, "threshold", 0.0);
+				column.applyEndOffset = columnEl.GetSafeAttribute<bool>(vars, "applyEndOffset", true);
+				column.frameExponent = columnEl.GetSafeAttribute<double>(vars, "frameExponent", 4.0);
+				column.statExponent = columnEl.GetSafeAttribute<double>(vars, "statExponent", 0.25);
 				columns.Add(column);
 			}
 		}
diff --git a/Engine/Source/Programs/CSVTools/PerfReportTool/Summaries/BucketSummary.cs b/Engine/Source/Programs/CSVTools/PerfReportTool/Summaries/BucketSummary.cs
--- a/Engine/Source/Programs/CSVTools/PerfReportTool/Summaries/BucketSummary.cs
+++ b/Engine/Source/Programs/CSVTools/PerfReportTool/Summaries/BucketSummary.cs
@@ -12,15 +12,15 @@ namespace PerfSummaries
 {
 	class BucketSummary : Summary
 	{
-		public BucketSummary(XElement element, string baseXmlDirectory)
+		public BucketSummary(XElement element, XmlVariableMappings vars, string baseXmlDirectory)
 		{
-			ReadStatsFromXML(element);
-			Title = element.GetSafeAttibute("title", "Untitled Bucket Summary");
+			ReadStatsFromXML(element, vars);
+			Title = element.GetSafeAttribute(vars, "title", "Untitled Bucket Summary");
 
 			XElement BucketElement = element.Element("buckets");
-			bool ReportOutOfRangeDefault = BucketElement.GetSafeAttibute<bool>("reportOutOfRange", false);
-			ReportBelowRange = BucketElement.GetSafeAttibute<bool>("reportBelowRange", ReportOutOfRangeDefault);
-			ReportAboveRange = BucketElement.GetSafeAttibute<bool>("reportAboveRange", ReportOutOfRangeDefault);
+			bool ReportOutOfRangeDefault = BucketElement.GetSafeAttribute<bool>(vars, "reportOutOfRange", false);
+			ReportBelowRange = BucketElement.GetSafeAttribute<bool>(vars, "reportBelowRange", ReportOutOfRangeDefault);
+			ReportAboveRange = BucketElement.GetSafeAttribute<bool>(vars, "reportAboveRange", ReportOutOfRangeDefault);
 
 			string[] XmlBuckets = BucketElement.Value.Split(',');
 			Buckets = new double[XmlBuckets.Length];
@@ -34,8 +34,8 @@ namespace PerfSummaries
 			XElement ColorElement = element.Element("colorDisplay");
 			if (ColorElement != null)
 			{
-				LowEndColor = new Colour(ColorElement.GetSafeAttibute<string>("lowEndColor"));
-				HighEndColor = new Colour(ColorElement.GetSafeAttibute<string>("highEndColor"));
+				LowEndColor = new Colour(ColorElement.GetSafeAttribute<string>(vars, "lowEndColor"));
+				HighEndColor = new Colour(ColorElement.GetSafeAttribute<string>(vars, "highEndColor"));
 			}
 		}
 		public BucketSummary() { }
diff --git a/Engine/Source/Programs/CSVTools/PerfReportTool/Summaries/EventSummary.cs b/Engine/Source/Programs/CSVTools/PerfReportTool/Summaries/EventSummary.cs
--- a/Engine/Source/Programs/CSVTools/PerfReportTool/Summaries/EventSummary.cs
+++ b/Engine/Source/Programs/CSVTools/PerfReportTool/Summaries/EventSummary.cs
@@ -11,12 +11,12 @@ namespace PerfSummaries
 {
 	class EventSummary : Summary
 	{
-		public EventSummary(XElement element, string baseXmlDirectory)
+		public EventSummary(XElement element, XmlVariableMappings vars, string baseXmlDirectory)
 		{
-			title = element.GetSafeAttibute("title", "Events");
-			summaryStatName = element.Attribute("summaryStatName").Value;
+			title = element.GetSafeAttribute(vars, "title", "Events");
+			summaryStatName = element.GetRequiredAttribute<string>(vars, "summaryStatName");
 			events = element.Element("events").Value.Split(',');
-			colourThresholds = ReadColourThresholdsXML(element.Element("colourThresholds"));
+			colourThresholds = ReadColourThresholdsXML(element.Element("colourThresholds"), vars);
 		}
 
 		public EventSummary() { }
diff --git a/Engine/Source/Programs/CSVTools/PerfReportTool/Summaries/ExtraLinksSummary.cs b/Engine/Source/Programs/CSVTools/PerfReportTool/Summaries/ExtraLinksSummary.cs
--- a/Engine/Source/Programs/CSVTools/PerfReportTool/Summaries/ExtraLinksSummary.cs
+++ b/Engine/Source/Programs/CSVTools/PerfReportTool/Summaries/ExtraLinksSummary.cs
@@ -64,12 +64,12 @@ namespace PerfSummaries
 			public string LinkURL;
             public string LinkTemplateCsvId;
         };
-        public ExtraLinksSummary(XElement element, string baseXmlDirectory, bool bInLinkTemplates=false)
+        public ExtraLinksSummary(XElement element, XmlVariableMappings vars, string baseXmlDirectory, bool bInLinkTemplates=false)
         {
             title = "Links";
             if (element != null)
             {
-                title = element.GetSafeAttibute("title", title);
+                title = element.GetSafeAttribute(vars, "title", title);
             }
             bLinkTemplates = bInLinkTemplates;
         }
diff --git a/Engine/Source/Programs/CSVTools/PerfReportTool/Summaries/FPSChartSummary.cs b/Engine/Source/Programs/CSVTools/PerfReportTool/Summaries/FPSChartSummary.cs
--- a/Engine/Source/Programs/CSVTools/PerfReportTool/Summaries/FPSChartSummary.cs
+++ b/Engine/Source/Programs/CSVTools/PerfReportTool/Summaries/FPSChartSummary.cs
@@ -12,20 +12,20 @@ namespace PerfSummaries
 {
 	class FPSChartSummary : Summary
 	{
-		public FPSChartSummary(XElement element, string baseXmlDirectory)
+		public FPSChartSummary(XElement element, XmlVariableMappings vars, string baseXmlDirectory)
 		{
-			ReadStatsFromXML(element);
-			fps = Convert.ToInt32(element.Attribute("fps").Value);
-			hitchThreshold = (float)Convert.ToDouble(element.Attribute("hitchThreshold").Value, System.Globalization.CultureInfo.InvariantCulture);
-			bUseEngineHitchMetric = element.GetSafeAttibute<bool>("useEngineHitchMetric", false);
+			ReadStatsFromXML(element, vars);
+			fps = element.GetRequiredAttribute<int>(vars, "fps");
+			hitchThreshold = (float)element.GetRequiredAttribute<double>(vars, "hitchThreshold");
+			bUseEngineHitchMetric = element.GetSafeAttribute<bool>(vars, "useEngineHitchMetric", false);
 			if (bUseEngineHitchMetric)
 			{
-				engineHitchToNonHitchRatio = element.GetSafeAttibute<float>("engineHitchToNonHitchRatio", 1.5f);
-				engineMinTimeBetweenHitchesMs = element.GetSafeAttibute<float>("engineMinTimeBetweenHitchesMs", 200.0f);
+				engineHitchToNonHitchRatio = element.GetSafeAttribute<float>(vars, "engineHitchToNonHitchRatio", 1.5f);
+				engineMinTimeBetweenHitchesMs = element.GetSafeAttribute<float>(vars, "engineMinTimeBetweenHitchesMs", 200.0f);
 			}
 
-			bIgnoreHitchTimePercent = element.GetSafeAttibute<bool>("ignoreHitchTimePercent", false);
-			bIgnoreMVP = element.GetSafeAttibute<bool>("ignoreMVP", false);
+			bIgnoreHitchTimePercent = element.GetSafeAttribute<bool>(vars, "ignoreHitchTimePercent", false);
+			bIgnoreMVP = element.GetSafeAttribute<bool>(vars, "ignoreMVP", false);
 		}
 
 		public FPSChartSummary() { }
diff --git a/Engine/Source/Programs/CSVTools/PerfReportTool/Summaries/HistogramSummary.cs b/Engine/Source/Programs/CSVTools/PerfReportTool/Summaries/HistogramSummary.cs
--- a/Engine/Source/Programs/CSVTools/PerfReportTool/Summaries/HistogramSummary.cs
+++ b/Engine/Source/Programs/CSVTools/PerfReportTool/Summaries/HistogramSummary.cs
@@ -11,11 +11,11 @@ namespace PerfSummaries
 {
 	class HistogramSummary : Summary
 	{
-		public HistogramSummary(XElement element, string baseXmlDirectory)
+		public HistogramSummary(XElement element, XmlVariableMappings vars, string baseXmlDirectory)
 		{
-			ReadStatsFromXML(element);
+			ReadStatsFromXML(element, vars);
 
-			ColourThresholds = ReadColourThresholdsXML(element.Element("colourThresholds"));
+			ColourThresholds = ReadColourThresholdsXML(element.Element("colourThresholds"), vars);
 
 			string[] histogramStrings = element.Element("histogramThresholds").Value.Split(',');
 			HistogramThresholds = new double[histogramStrings.Length];
@@ -28,12 +28,12 @@ namespace PerfSummaries
 			{
 				if (child.Name == "budgetOverride")
 				{
-					BudgetOverrideStatName = child.Attribute("stat").Value;
-					BudgetOverrideStatBudget = Convert.ToDouble(child.Attribute("budget").Value, System.Globalization.CultureInfo.InvariantCulture);
+					BudgetOverrideStatName = child.GetRequiredAttribute<string>(vars, "stat");
+					BudgetOverrideStatBudget = child.GetRequiredAttribute<double>(vars, "budget");
 				}
 			}
 
-			bSuppressAveragesTable = element.GetSafeAttibute<bool>("suppressAveragesTable", false);
+			bSuppressAveragesTable = element.GetSafeAttribute<bool>(vars, "suppressAveragesTable", false);
 		}
 
 		public HistogramSummary() { }
diff --git a/Engine/Source/Programs/CSVTools/PerfReportTool/Summaries/HitchSummary.cs b/Engine/Source/Programs/CSVTools/PerfReportTool/Summaries/HitchSummary.cs
--- a/Engine/Source/Programs/CSVTools/PerfReportTool/Summaries/HitchSummary.cs
+++ b/Engine/Source/Programs/CSVTools/PerfReportTool/Summaries/HitchSummary.cs
@@ -12,9 +12,9 @@ namespace PerfSummaries
 {
 	class HitchSummary : Summary
 	{
-		public HitchSummary(XElement element, string baseXmlDirectory)
+		public HitchSummary(XElement element, XmlVariableMappings vars, string baseXmlDirectory)
 		{
-			ReadStatsFromXML(element);
+			ReadStatsFromXML(element, vars);
 
 			string[] hitchThresholds = element.Element("hitchThresholds").Value.Split(',');
 			HitchThresholds = new double[hitchThresholds.Length];
diff --git a/Engine/Source/Programs/CSVTools/PerfReportTool/Summaries/MapOverlaySummary.cs b/Engine/Source/Programs/CSVTools/PerfReportTool/Summaries/MapOverlaySummary.cs
--- a/Engine/Source/Programs/CSVTools/PerfReportTool/Summaries/MapOverlaySummary.cs
+++ b/Engine/Source/Programs/CSVTools/PerfReportTool/Summaries/MapOverlaySummary.cs
@@ -32,20 +32,20 @@ namespace PerfSummaries
 
 		class MapOverlay
 		{
-			public MapOverlay(XElement element)
+			public MapOverlay(XElement element, XmlVariableMappings vars)
 			{
-				positionStatNames[0] = element.GetSafeAttibute<string>("xStat");
-				positionStatNames[1] = element.GetSafeAttibute<string>("yStat");
-				positionStatNames[2] = element.GetSafeAttibute<string>("zStat");
-				summaryStatNamePrefix = element.GetSafeAttibute<string>("summaryStatNamePrefix"); // unused!
-				lineColor = element.GetSafeAttibute<string>("lineColor", "#ffffff");
+				positionStatNames[0] = element.GetSafeAttribute<string>(vars, "xStat");
+				positionStatNames[1] = element.GetSafeAttribute<string>(vars, "yStat");
+				positionStatNames[2] = element.GetSafeAttribute<string>(vars, "zStat");
+				summaryStatNamePrefix = element.GetSafeAttribute<string>(vars, "summaryStatNamePrefix"); // unused!
+				lineColor = element.GetSafeAttribute<string>(vars, "lineColor", "#ffffff");
 				foreach (XElement eventEl in element.Elements("event"))
 				{
-					MapOverlayEvent ev = new MapOverlayEvent(eventEl.Attribute("name").Value);
-					ev.shortName = eventEl.GetSafeAttibute<string>("shortName");
-					ev.summaryStatName = eventEl.GetSafeAttibute<string>("summaryStatName"); // unused!
-					ev.lineColor = eventEl.GetSafeAttibute<string>("lineColor");
-					if (eventEl.GetSafeAttibute<bool>("isStartEvent", false))
+					MapOverlayEvent ev = new MapOverlayEvent(eventEl.GetRequiredAttribute<string>(vars, "name"));
+					ev.shortName = eventEl.GetSafeAttribute<string>(vars, "shortName");
+					ev.summaryStatName = eventEl.GetSafeAttribute<string>(vars, "summaryStatName"); // unused!
+					ev.lineColor = eventEl.GetSafeAttribute<string>(vars, "lineColor");
+					if (eventEl.GetSafeAttribute<bool>(vars, "isStartEvent", false))
 					{
 						if (startEvent != null)
 						{
@@ -64,15 +64,15 @@ namespace PerfSummaries
 			public List<MapOverlayEvent> events = new List<MapOverlayEvent>();
 		}
 
-		public MapOverlaySummary(XElement element, string baseXmlDirectory)
+		public MapOverlaySummary(XElement element, XmlVariableMappings vars, string baseXmlDirectory)
 		{
-			ReadStatsFromXML(element);
+			ReadStatsFromXML(element, vars);
 			if (stats.Count != 0)
 			{
 				throw new Exception("<stats> element is not supported");
 			}
 
-			sourceImagePath = element.GetSafeAttibute<string>("sourceImage");
+			sourceImagePath = element.GetSafeAttribute<string>(vars, "sourceImage");
 			if (baseXmlDirectory == null)
 			{
 				throw new Exception("BaseXmlDirectory not specified");
@@ -82,19 +82,19 @@ namespace PerfSummaries
 				sourceImagePath = System.IO.Path.GetFullPath(System.IO.Path.Combine(baseXmlDirectory, sourceImagePath));
 			}
 
-			offsetX = element.GetSafeAttibute<float>("offsetX", 0.0f);
-			offsetY = element.GetSafeAttibute<float>("offsetY", 0.0f);
-			scale = element.GetSafeAttibute<float>("scale", 1.0f);
-			title = element.GetSafeAttibute("title", "Events");
-			destImageFilename = element.Attribute("destImage").Value;
-			imageWidth = element.GetSafeAttibute<float>("width", 250.0f);
-			imageHeight = element.GetSafeAttibute<float>("height", 250.0f);
-			framesPerLineSegment = element.GetSafeAttibute<int>("framesPerLineSegment", 5);
-			lineSplitDistanceThreshold = element.GetSafeAttibute<float>("lineSplitDistanceThreshold", float.MaxValue);
+			offsetX = element.GetSafeAttribute<float>(vars, "offsetX", 0.0f);
+			offsetY = element.GetSafeAttribute<float>(vars, "offsetY", 0.0f);
+			scale = element.GetSafeAttribute<float>(vars, "scale", 1.0f);
+			title = element.GetSafeAttribute(vars, "title", "Events");
+			destImageFilename = element.GetRequiredAttribute<string>(vars, "destImage");
+			imageWidth = element.GetSafeAttribute<float>(vars, "width", 250.0f);
+			imageHeight = element.GetSafeAttribute<float>(vars, "height", 250.0f);
+			framesPerLineSegment = element.GetSafeAttribute<int>(vars, "framesPerLineSegment", 5);
+			lineSplitDistanceThreshold = element.GetSafeAttribute<float>(vars, "lineSplitDistanceThreshold", float.MaxValue);
 
 			foreach (XElement overlayEl in element.Elements("overlay"))
 			{
-				MapOverlay overlay = new MapOverlay(overlayEl);
+				MapOverlay overlay = new MapOverlay(overlayEl, vars);
 				overlays.Add(overlay);
 				stats.Add(overlay.positionStatNames[0]);
 				stats.Add(overlay.positionStatNames[1]);
diff --git a/Engine/Source/Programs/CSVTools/PerfReportTool/Summaries/PeakSummary.cs b/Engine/Source/Programs/CSVTools/PerfReportTool/Summaries/PeakSummary.cs
--- a/Engine/Source/Programs/CSVTools/PerfReportTool/Summaries/PeakSummary.cs
+++ b/Engine/Source/Programs/CSVTools/PerfReportTool/Summaries/PeakSummary.cs
@@ -32,10 +32,10 @@ namespace PerfSummaries
 				title = inTitle;
 				statNamesFilter = inStatFilterStr.Split(',');
 			}
-			public PeakSummarySection(XElement element)
+			public PeakSummarySection(XElement element, XmlVariableMappings vars)
 			{
 				XElement statFilterElement = element.Element("statFilter");
-				title = element.Attribute("title").Value;
+				title = element.GetRequiredAttribute<string>(vars, "title");
 
 				statNamesFilter = statFilterElement.Value.Split(',');
 			}
@@ -58,17 +58,17 @@ namespace PerfSummaries
 			public string title;
 		};
 
-		public PeakSummary(XElement element, string baseXmlDirectory)
+		public PeakSummary(XElement element, XmlVariableMappings vars, string baseXmlDirectory)
 		{
 			//read the child elements (mostly for colourThresholds)
-			ReadStatsFromXML(element);
-			hideStatPrefix = element.GetSafeAttibute<string>("hideStatPrefix", "").ToLower();
+			ReadStatsFromXML(element, vars);
+			hideStatPrefix = element.GetSafeAttribute<string>(vars, "hideStatPrefix", "").ToLower();
 
 			foreach (XElement child in element.Elements())
 			{
 				if (child.Name == "summarySection")
 				{
-					peakSummarySections.Add(new PeakSummarySection(child));
+					peakSummarySections.Add(new PeakSummarySection(child, vars));
 				}
 			}
 
@@ -213,7 +213,7 @@ namespace PerfSummaries
 
 
 
-		void AddStat(string statName, OptionalDouble budget)
+		void AddStat(string statName, Optional<double> budget)
 		{
 			stats.Add(statName);
 
@@ -258,13 +258,13 @@ namespace PerfSummaries
 		{
 			public PeakStatInfo(string inName, string inShortName)
 			{
-				budget = new OptionalDouble();
+				budget = new Optional<double>();
 				name = inName;
 				shortName = inShortName;
 			}
 			public string name;
 			public string shortName;
-			public OptionalDouble budget;
+			public Optional<double> budget;
 		};
 
 		PeakStatInfo getOrAddStatInfo(string statName)
diff --git a/Engine/Source/Programs/CSVTools/PerfReportTool/Summary.cs b/Engine/Source/Programs/CSVTools/PerfReportTool/Summary.cs
--- a/Engine/Source/Programs/CSVTools/PerfReportTool/Summary.cs
+++ b/Engine/Source/Programs/CSVTools/PerfReportTool/Summary.cs
@@ -27,11 +27,11 @@ namespace PerfSummaries
 			}
 		}
 
-		public static Summary Create(string summaryTypeName, XElement summaryXmlElement, string baseXmlDirectory)
+		public static Summary Create(string summaryTypeName, XElement summaryXmlElement, XmlVariableMappings vars, string baseXmlDirectory)
 		{
 			if ( summaryNameLookup.TryGetValue(summaryTypeName, out System.Type summaryType ) )
 			{
-				object[] constructArgs = new object[2] { summaryXmlElement, baseXmlDirectory };
+				object[] constructArgs = new object[3] { summaryXmlElement, vars, baseXmlDirectory };
 				return (Summary)Activator.CreateInstance(summaryType, constructArgs);
 			}
 			throw new Exception("Summary type " + summaryType + " not found!");
@@ -90,9 +90,9 @@ namespace PerfSummaries
 
 		public abstract string GetName();
 
-        public void ReadStatsFromXML(XElement element)
+        public void ReadStatsFromXML(XElement element, XmlVariableMappings vars)
         {
-			useUnstrippedCsvStats = element.GetSafeAttibute<bool>("useUnstrippedCsvStats", false);
+			useUnstrippedCsvStats = element.GetSafeAttribute<bool>(vars, "useUnstrippedCsvStats", false);
 			XElement statsElement = element.Element("stats");
 			if (statsElement != null)
 			{
@@ -102,11 +102,11 @@ namespace PerfSummaries
             {
                 if (child.Name == "capture")
                 {
-                    string captureName = child.Attribute("name").Value;
-                    string captureStart = child.Attribute("startEvent").Value;
-                    string captureEnd = child.Attribute("endEvent").Value;
-                    bool incFirstFrame = Convert.ToBoolean(child.Attribute("includeFirstFrame").Value);
-                    bool incLastFrame = Convert.ToBoolean(child.Attribute("includeLastFrame").Value);
+                    string captureName = child.GetRequiredAttribute<string>(vars, "name");
+                    string captureStart = child.GetRequiredAttribute<string>(vars, "startEvent");
+                    string captureEnd = child.GetRequiredAttribute<string>(vars, "endEvent");
+					bool incFirstFrame = child.GetSafeAttribute<bool>(vars, "includeFirstFrame", true);
+					bool incLastFrame = child.GetSafeAttribute<bool>(vars, "includeLastFrame", true);
                     CaptureRange newRange = new CaptureRange(captureName, captureStart, captureEnd);
                     newRange.includeFirstFrame = incFirstFrame;
                     newRange.includeLastFrame = incLastFrame;
@@ -118,8 +118,9 @@ namespace PerfSummaries
                     {
                         continue;
                     }
-                    string statName = child.Attribute("stat").Value;
-                    string[] hitchThresholdsStrList = child.Value.Split(',');
+                    string statName = child.GetRequiredAttribute<string>(vars, "stat");
+					string hitchThresholdsStr = child.GetValue(vars);
+                    string[] hitchThresholdsStrList = hitchThresholdsStr.Split(',');
 					ColourThresholdList HitchThresholds = new ColourThresholdList();
 					for (int i = 0; i < hitchThresholdsStrList.Length; i++)
                     {
@@ -133,7 +134,7 @@ namespace PerfSummaries
 						{
 							hitchThresholdNumStr = hitchThresholdStr.Substring(0, openBracketIndex);
 							int closeBracketIndex = hitchThresholdStr.IndexOf(')');
-							if (closeBracketIndex > openBracketIndex)
+							if (closeBracketIndex > openBracketIndex) 
 							{
 								string colourString = hitchThresholdStr.Substring(openBracketIndex+1, closeBracketIndex - openBracketIndex-1);
 								thresholdColour = new Colour(colourString);
@@ -196,14 +197,14 @@ namespace PerfSummaries
             return uniqueStats.ToArray();
         }
 
-		protected ColourThresholdList ReadColourThresholdListXML(XElement colourThresholdEl)
+		protected ColourThresholdList ReadColourThresholdListXML(XElement colourThresholdEl, XmlVariableMappings vars)
 		{
-			return ColourThresholdList.ReadColourThresholdListXML(colourThresholdEl);
+			return ColourThresholdList.ReadColourThresholdListXML(colourThresholdEl, vars);
 		}
 
-		protected double [] ReadColourThresholdsXML(XElement colourThresholdEl)
+		protected double [] ReadColourThresholdsXML(XElement colourThresholdEl, XmlVariableMappings vars)
 		{
-			return ColourThresholdList.ReadColourThresholdsXML(colourThresholdEl);
+			return ColourThresholdList.ReadColourThresholdsXML(colourThresholdEl, vars);
 		}
 
         public string GetStatThresholdColour(string StatToUse, double value)
diff --git a/Engine/Source/Programs/CSVTools/PerfReportTool/SummaryTable.cs b/Engine/Source/Programs/CSVTools/PerfReportTool/SummaryTable.cs
--- a/Engine/Source/Programs/CSVTools/PerfReportTool/SummaryTable.cs
+++ b/Engine/Source/Programs/CSVTools/PerfReportTool/SummaryTable.cs
@@ -95,10 +95,10 @@ namespace PerfSummaries
 				columnFilterList.AddRange(appendList);
 			}
 
-			bReverseSortRows = tableElement.GetSafeAttibute<bool>("reverseSortRows", false);
-			bScrollableFormatting = tableElement.GetSafeAttibute<bool>("scrollableFormatting", false);
+			bReverseSortRows = tableElement.GetSafeAttribute<bool>("reverseSortRows", false);
+			bScrollableFormatting = tableElement.GetSafeAttribute<bool>("scrollableFormatting", false);
 
-			string colorizeModeStr = tableElement.GetSafeAttibute<string>("colorizeMode", "").ToLower();
+			string colorizeModeStr = tableElement.GetSafeAttribute<string>("colorizeMode", "").ToLower();
 			if (colorizeModeStr != "")
 			{
 				if (colorizeModeStr == "auto")
@@ -115,22 +115,22 @@ namespace PerfSummaries
 				}
 			}
 
-			statThreshold = tableElement.GetSafeAttibute<float>("statThreshold", 0.0f);
-			hideStatPrefix = tableElement.GetSafeAttibute<string>("hideStatPrefix");
+			statThreshold = tableElement.GetSafeAttribute<float>("statThreshold", 0.0f);
+			hideStatPrefix = tableElement.GetSafeAttribute<string>("hideStatPrefix");
 
 			foreach (XElement sectionBoundaryEl in tableElement.Elements("sectionBoundary"))
 			{
 				if (sectionBoundaryEl != null)
 				{
-					string statName = ApplySubstitution(sectionBoundaryEl.GetSafeAttibute<string>("statName"), substitutionsDict);
+					string statName = ApplySubstitution(sectionBoundaryEl.GetSafeAttribute<string>("statName"), substitutionsDict);
 
 					SummarySectionBoundaryInfo sectionBoundary = new SummarySectionBoundaryInfo(
 						statName,
-						sectionBoundaryEl.GetSafeAttibute<string>("startToken"),
-						sectionBoundaryEl.GetSafeAttibute<string>("endToken"),
-						sectionBoundaryEl.GetSafeAttibute<int>("level", 0),
-						sectionBoundaryEl.GetSafeAttibute<bool>("inCollatedTable", true),
-						sectionBoundaryEl.GetSafeAttibute<bool>("inFullTable", true)
+						sectionBoundaryEl.GetSafeAttribute<string>("startToken"),
+						sectionBoundaryEl.GetSafeAttribute<string>("endToken"),
+						sectionBoundaryEl.GetSafeAttribute<int>("level", 0),
+						sectionBoundaryEl.GetSafeAttribute<bool>("inCollatedTable", true),
+						sectionBoundaryEl.GetSafeAttribute<bool>("inFullTable", true)
 						);
 					sectionBoundaries.Add(sectionBoundary);
 				}
@@ -268,7 +268,7 @@ namespace PerfSummaries
 		{
 			name = element.Attribute("name").Value.ToLower();
 
-			string autoColorizeStr = element.GetSafeAttibute<string>("autoColorize", "highIsBad").ToLower();
+			string autoColorizeStr = element.GetSafeAttribute<string>("autoColorize", "highIsBad").ToLower();
 			var modeList = Enum.GetValues(typeof(AutoColorizeMode));
 			foreach (AutoColorizeMode mode in modeList)
 			{
@@ -278,20 +278,20 @@ namespace PerfSummaries
 					break;
 				}
 			}
-			numericFormat = element.GetSafeAttibute<string>("numericFormat");
-			maxStringLength = element.GetSafeAttibute<int>("maxStringLength", Int32.MaxValue );
-			maxStringLengthCollated = element.GetSafeAttibute<int>("maxStringLengthCollated", Int32.MaxValue );
+			numericFormat = element.GetSafeAttribute<string>("numericFormat");
+			maxStringLength = element.GetSafeAttribute<int>("maxStringLength", Int32.MaxValue );
+			maxStringLengthCollated = element.GetSafeAttribute<int>("maxStringLengthCollated", Int32.MaxValue );
 			if (maxStringLengthCollated == Int32.MaxValue)
 			{
 				maxStringLengthCollated = maxStringLength;
 			}
 
-			noWrap = element.GetSafeAttibute<string>("noWrap") == "true";
+			noWrap = element.GetSafeAttribute<string>("noWrap") == "true";
 
 			if (IsDate())
 			{
-				dateFormat = element.GetSafeAttibute<string>("dateFormat");
-				string timeZoneId = element.GetSafeAttibute<string>("dateTimeZoneId");
+				dateFormat = element.GetSafeAttribute<string>("dateFormat");
+				string timeZoneId = element.GetSafeAttribute<string>("dateTimeZoneId");
 				dateTimeZone = TimeZoneInfo.Utc;
 				if (timeZoneId != null)
 				{
@@ -301,13 +301,13 @@ namespace PerfSummaries
 				}
 			}
 
-			includeValueWithBucketName = element.GetSafeAttibute<bool>("includeValueWithBucketName", true);
-			string bucketNamesString = element.GetSafeAttibute<string>("valueBucketNames");
+			includeValueWithBucketName = element.GetSafeAttribute<bool>("includeValueWithBucketName", true);
+			string bucketNamesString = element.GetSafeAttribute<string>("valueBucketNames");
 			if (bucketNamesString != null)
 			{
 				bucketNames = bucketNamesString.Split(',').ToList();
 			}
-			string bucketThresholdsString = element.GetSafeAttibute<string>("valueBucketThresholds");
+			string bucketThresholdsString = element.GetSafeAttribute<string>("valueBucketThresholds");
 			if (bucketThresholdsString != null)
 			{
 				bucketThresholds = bucketThresholdsString.Split(',').Select(valStr =>
@@ -320,7 +320,7 @@ namespace PerfSummaries
 				}).ToList();
 			}
 
-			colourThresholdList = ColourThresholdList.ReadColourThresholdListXML(element.Element("colourThresholds"));
+			colourThresholdList = ColourThresholdList.ReadColourThresholdListXML(element.Element("colourThresholds"), null);
 		}
 
 		public bool IsDate() => numericFormat == "date";
diff --git a/Engine/Source/Programs/CSVTools/PerfReportTool/SummaryTableData.cs b/Engine/Source/Programs/CSVTools/PerfReportTool/SummaryTableData.cs
--- a/Engine/Source/Programs/CSVTools/PerfReportTool/SummaryTableData.cs
+++ b/Engine/Source/Programs/CSVTools/PerfReportTool/SummaryTableData.cs
@@ -119,6 +119,10 @@ namespace PerfSummaries
 
 		public SummaryTableElement(Type inType, string inName, double inValue, ColourThresholdList inColorThresholdList, string inToolTip, uint inFlags = 0)
 		{
+			if (inName == "")
+			{
+				throw new Exception("Name cannot be empty!");
+			}
 			type = inType;
 			name = inName;
 			isNumeric = true;
@@ -130,6 +134,10 @@ namespace PerfSummaries
 		}
 		public SummaryTableElement(Type inType, string inName, string inValue, ColourThresholdList inColorThresholdList, string inToolTip, uint inFlags = 0)
 		{
+			if (inName == "")
+			{
+				throw new Exception("Name cannot be empty!");
+			}
 			type = inType;
 			name = inName;
 			numericValue = 0.0;
diff --git a/Engine/Source/Programs/CSVTools/PerfReportTool/SummaryTypes.cs b/Engine/Source/Programs/CSVTools/PerfReportTool/SummaryTypes.cs
--- a/Engine/Source/Programs/CSVTools/PerfReportTool/SummaryTypes.cs
+++ b/Engine/Source/Programs/CSVTools/PerfReportTool/SummaryTypes.cs
@@ -1,5 +1,6 @@
 // Copyright Epic Games, Inc. All Rights Reserved.
 
+using PerfReportTool;
 using System;
 using System.Collections.Generic;
 using System.Linq;
@@ -382,24 +383,21 @@ namespace PerfSummaries
 			return list.GetColourForValue(value);
 		}
 
-		public static ColourThresholdList ReadColourThresholdListXML(XElement colourThresholdEl)
+		public static ColourThresholdList ReadColourThresholdListXML(XElement colourThresholdEl, XmlVariableMappings vars)
 		{
 			if (colourThresholdEl != null)
 			{
 				ColourThresholdList thresholdList = new ColourThresholdList();
 
-				XAttribute lerpColoursElement = colourThresholdEl.Attribute("lerpColours");
-				if (lerpColoursElement != null)
-				{
-					thresholdList.LerpColours = bool.Parse(lerpColoursElement.Value);
-				}
+				thresholdList.LerpColours = colourThresholdEl.GetSafeAttribute<bool>(vars, "lerpColours", thresholdList.LerpColours);
 
 				XElement thresholdsElement = colourThresholdEl.Element("thresholds");
 				XElement coloursElement = colourThresholdEl.Element("colours");
 				if (thresholdsElement != null && coloursElement != null)
 				{
-					string[] thresholdStrings = thresholdsElement.Value.Split(',');
-					string[] colourStrings = coloursElement.Value.Split(',');
+					// Colour thresholds support variables so resolve them if provided
+					string[] thresholdStrings = thresholdsElement.GetValue(vars).Split(',');
+					string[] colourStrings = coloursElement.GetValue(vars).Split(',');
 					if (thresholdStrings.Length != colourStrings.Length)
 					{
 						throw new Exception("Number of colour thresholds and colours must match");
@@ -415,7 +413,7 @@ namespace PerfSummaries
 				else
 				{
 					// Revert to previous logic if sub elements don't exist.
-					double[] values = ReadColourThresholdsXML(colourThresholdEl);
+					double[] values = ReadColourThresholdsXML(colourThresholdEl, vars);
 					if (values != null)
 					{
 						foreach (double value in values)
@@ -429,11 +427,12 @@ namespace PerfSummaries
 			return null;
 		}
 
-		public static double[] ReadColourThresholdsXML(XElement colourThresholdEl)
+		public static double[] ReadColourThresholdsXML(XElement colourThresholdEl, XmlVariableMappings vars )
 		{
 			if (colourThresholdEl != null)
 			{
-				string[] colourStrings = colourThresholdEl.Value.Split(',');
+				string colourThresholdElementStr = colourThresholdEl.GetValue(vars);
+				string[] colourStrings = colourThresholdElementStr.Split(',');
 				if (colourStrings.Length != 4)
 				{
 					throw new Exception("Incorrect number of colourthreshold entries. Should be 4.");
