PerfReportTool 4.210.1 : fix an exception when startTimestamp is not set for some rows

diff --git a/Engine/Source/Programs/CSVTools/PerfReportTool/PerfReportTool.cs b/Engine/Source/Programs/CSVTools/PerfReportTool/PerfReportTool.cs
--- a/Engine/Source/Programs/CSVTools/PerfReportTool/PerfReportTool.cs
+++ b/Engine/Source/Programs/CSVTools/PerfReportTool/PerfReportTool.cs
@@ -23,7 +23,7 @@ namespace PerfReportTool
     class Version
     {
 		// Format: Major.Minor.Bugfix
-        private static string VersionString = "4.210.0";
+        private static string VersionString = "4.210.1";
 
         public static string Get() { return VersionString; }
     };
diff --git a/Engine/Source/Programs/CSVTools/PerfReportTool/SummaryTable.cs b/Engine/Source/Programs/CSVTools/PerfReportTool/SummaryTable.cs
--- a/Engine/Source/Programs/CSVTools/PerfReportTool/SummaryTable.cs
+++ b/Engine/Source/Programs/CSVTools/PerfReportTool/SummaryTable.cs
@@ -911,7 +911,7 @@ namespace PerfSummaries
 
 				if (forceNumericFormat != null)
 				{
-					if (forceNumericFormat == "date")
+					if (forceNumericFormat == "date" && val != double.MaxValue)
 					{
 						DateTimeOffset dateTimeOffset = DateTimeOffset.FromUnixTimeSeconds((long)val);
 						TimeSpan timeZoneOffset = formatInfo.dateTimeZone.GetUtcOffset(dateTimeOffset);
@@ -1903,8 +1903,12 @@ namespace PerfSummaries
 							if (column.isNumeric && column.formatInfo != null && column.formatInfo.IsDate())
 							{
 								// For dates use a standard UTC timestamp for the tooltip
-								DateTimeOffset dateTimeOffset = DateTimeOffset.FromUnixTimeSeconds((long)column.GetValue(rowIndex));
-								toolTip = dateTimeOffset.ToString("yyyy-MM-dd HH:mm:ss (UTC)");
+								double val = column.GetValue(rowIndex);
+								if (val < double.MaxValue )
+								{
+									DateTimeOffset dateTimeOffset = DateTimeOffset.FromUnixTimeSeconds((long)val);
+									toolTip = dateTimeOffset.ToString("yyyy-MM-dd HH:mm:ss (UTC)");
+								}
 							}
 							else
 							{
