Improvements to metadata query evaluation (csvStats lib)
  -  Equallity now uses numeric comparisons if appropriate
  - Add a ~= operator for contains
  - Fix a potential infinite loop if a bad term is evaluated

diff --git a/Engine/Source/Programs/CSVTools/CsvStats/MetadataQueryBuilder.cs b/Engine/Source/Programs/CSVTools/CsvStats/MetadataQueryBuilder.cs
--- a/Engine/Source/Programs/CSVTools/CsvStats/MetadataQueryBuilder.cs
+++ b/Engine/Source/Programs/CSVTools/CsvStats/MetadataQueryBuilder.cs
@@ -21,6 +21,7 @@ namespace CSVStats
 		LessThan,
 		GreaterThanOrEqual,
 		LessThanOrEqual,
+		Contains,
 
 		COUNT
 	};
@@ -42,28 +43,45 @@ namespace CSVStats
 			}
 			bool matches = false;
 			bool negate = bNegate;
+			string metaDataValueStr = metadata.Values[Key].ToLower();
+
+			// Attempt to parse both values as numeric and determine if this is a numeric comparison
+			double doubleVal = double.MaxValue;
+			bool bNumericComparison = Double.TryParse(Value, out doubleVal);
+			double doubleMetadataVal = double.MaxValue;
+			bNumericComparison &= Double.TryParse(metaDataValueStr, out doubleMetadataVal);
+
 			if (Op == ComparisonOp.Equal || Op == ComparisonOp.NotEqual)
 			{
 				// Check if the value actually matches (allow wildcards)
-				matches = CsvStats.DoesSearchStringMatch(metadata.Values[Key].ToLower(), Value.ToLower());
-				if (Op == ComparisonOp.NotEqual)
+				if (bNumericComparison)
 				{
-					negate = !negate;
+					matches = doubleVal == doubleMetadataVal;
 				}
-			}
-			if (Op == ComparisonOp.LessThan || Op == ComparisonOp.GreaterThan || Op == ComparisonOp.LessThanOrEqual || Op == ComparisonOp.GreaterThanOrEqual )
-			{
-				double doubleVal;
-				if ( !Double.TryParse(Value, out doubleVal ) )
+				else 
 				{
-					throw new Exception("Metadata operators expect numeric values!");
+					// Do a string comparison (taking wildcards into account)
+					matches = CsvStats.DoesSearchStringMatch(metaDataValueStr, Value.ToLower());
+					if (Op == ComparisonOp.NotEqual)
+					{
+						negate = !negate;
+					}
 				}
-				double doubleMetadataVal;
-				if (!Double.TryParse(metadata.Values[Key].ToLower(), out doubleMetadataVal))
+			}
+			else if (Op == ComparisonOp.LessThan || Op == ComparisonOp.GreaterThan || Op == ComparisonOp.LessThanOrEqual || Op == ComparisonOp.GreaterThanOrEqual)
+			{
+				if (!bNumericComparison)
 				{
-					// Metadata value wasn't numeric, so just return
-					Console.WriteLine("Value for "+Key+" was non-numeric: " + Value);
-					return false;
+					if (doubleVal == double.MaxValue)
+					{
+						throw new Exception("Metadata operators expect numeric values!");
+					}
+					if (doubleMetadataVal == double.MaxValue)
+					{
+						// Metadata value wasn't numeric, so just return
+						Console.WriteLine("Value for " + Key + " was non-numeric: " + Value);
+						return false;
+					}
 				}
 				switch (Op)
 				{
@@ -77,6 +95,10 @@ namespace CSVStats
 						return doubleMetadataVal >= doubleVal;
 				}
 			}
+			else if (Op == ComparisonOp.Contains)
+			{
+				matches = metaDataValueStr.Contains(Value);
+			}
 			return negate ? !matches : matches;
 		}
 		public string Key;
@@ -151,6 +173,9 @@ namespace CSVStats
 				case ">=":
 					op = ComparisonOp.GreaterThanOrEqual;
 					break;
+				case "~=":
+					op = ComparisonOp.Contains;
+					break;
 				default:
 					break;
 			}
@@ -246,7 +271,7 @@ namespace CSVStats
 					expression = ConsumeTerm(ref tokenList);
 					if (expression == null)
 					{
-						throw new Exception("Error parsing metadata filter expression");
+						throw new Exception("Error parsing metadata filter expression: "+string.Join("",tokenList));
 					}
 				}
 				expression.bNegate = bNegate;
@@ -294,8 +319,8 @@ namespace CSVStats
 
 		private static string ConsumeNextToken(ref string remainingString)
 		{
-			string[] operatorTokenArray = { "!=", "<=", ">=", "(", ")", "=", "<", ">" };
-			string operatorChars = "!=<>()";
+			string[] operatorTokenArray = { "!=", "<=", ">=", "(", ")", "=", "<", ">", "~=" };
+			string operatorChars = "!=<>()~";
 
 			// Consume leading whitespace
 			remainingString = remainingString.TrimStart();
@@ -318,7 +343,7 @@ namespace CSVStats
 			for (int i = 0; i < remainingString.Length; i++)
 			{
 				char c = remainingString[i];
-				if (Char.IsWhiteSpace(c) || operatorChars.IndexOf(c) >= 0)
+				if ( i>0 && ( Char.IsWhiteSpace(c) || operatorChars.IndexOf(c) >= 0 ))
 				{
 					break;
 				}
@@ -341,60 +366,6 @@ namespace CSVStats
 					tokens.Add(token);
 				}
 			}
-
-
-			/*
-			List<string> operatorTokens = new List<string>(operatorTokenArray);
-
-			// Find all operator chars
-			string operatorChars = "";
-			foreach (string opToken in operatorTokenArray)
-				foreach (char c in opToken)
-					if (operatorChars.IndexOf(c)<0)
-						operatorChars += c;
-
-			List<string> tokens = new List<string>();
-			string currentToken = "";
-			CharType prevCharType = CharType.Whitespace;
-			for (int i = 0; i < str.Length; i++)
-			{
-				char c = str[i];
-
-				// Find the character type
-				CharType newCharType;
-				if (Char.IsWhiteSpace(c))
-				{
-					newCharType = CharType.Whitespace;
-				}
-				else if (operatorChars.Contains(c.ToString()))
-				{
-					newCharType = CharType.Operator;
-				}
-				else
-				{
-					newCharType = CharType.Standard;
-				}
-
-				// If the character type changed, add the token
-				if (newCharType != prevCharType)
-				{
-					prevCharType = newCharType;
-					if (currentToken.Length > 0)
-					{
-						tokens.Add(currentToken);
-						currentToken = "";
-					}
-				}
-				if ( newCharType != CharType.Whitespace )
-				{
-					currentToken += c;
-				}
-			}
-			if (currentToken.Length > 0)
-			{
-				tokens.Add(currentToken);
-			}
-			*/
 			return tokens;
 		}
 
