PerfReportTool 4.220: VariableSets support
  - Add support for variable sets. This allows us to set blocks of variables based on metadata, which can be accessed via ${NAME} syntax. Variable sets support nesting and can be conditional using metadata queries
  - We support a global variable set (declared under <globalVariableSet> in root), and local variable sets which apply to a given report type only (as <variableSet> children of the ReportType element)

example:
<globalVariableSet>
  <variableSet>
    <var mvpStatName="MVP${meta.targetframerate}"/>
    <var targetFps="${meta.targetframerate}"/>
    <var dynamicResBudget="100"/>
  </variableSet>

  <!-- Switch -->
  <variableSet metadataQuery="platform=Switch">
    <var frameTimeBudget="33.33"/>
    <variableSet metadataQuery="deviceprofile=Switch_Console_Forward*">
      <var dynamicResBudget="92"/>
    </variableSet>
    <variableSet metadataQuery="deviceprofile=Switch_Handheld_Forward*">
      <var dynamicResBudget="88"/>
    </variableSet>
  </variableSet>
</globalVariableSet>

#tests Passes Test cases with identical results (if no changes to the XML)

diff --git a/Engine/Source/Programs/CSVTools/PerfReportTool/PerfReportTool.cs b/Engine/Source/Programs/CSVTools/PerfReportTool/PerfReportTool.cs
--- a/Engine/Source/Programs/CSVTools/PerfReportTool/PerfReportTool.cs
+++ b/Engine/Source/Programs/CSVTools/PerfReportTool/PerfReportTool.cs
@@ -23,7 +23,7 @@ namespace PerfReportTool
     class Version
     {
 		// Format: Major.Minor.Bugfix
-        private static string VersionString = "4.210.1";
+        private static string VersionString = "4.220.0";
 
         public static string Get() { return VersionString; }
     };
diff --git a/Engine/Source/Programs/CSVTools/PerfReportTool/ReportType.cs b/Engine/Source/Programs/CSVTools/PerfReportTool/ReportType.cs
--- a/Engine/Source/Programs/CSVTools/PerfReportTool/ReportType.cs
+++ b/Engine/Source/Programs/CSVTools/PerfReportTool/ReportType.cs
@@ -18,13 +18,19 @@ namespace PerfReportTool
 
 	class ReportTypeInfo
 	{
-		public ReportTypeInfo(XElement element, Dictionary<string, XElement> sharedSummaries, string baseXmlDirectory, XmlVariableMappings inVariableMappings)
+		public ReportTypeInfo(XElement element, Dictionary<string, XElement> sharedSummaries, string baseXmlDirectory, XmlVariableMappings inVariableMappings, CsvMetadata csvMetadata )
 		{
 			graphs = new List<ReportGraph>();
 			summaries = new List<Summary>();
 
 			vars = inVariableMappings;
 
+			// Apply local variable sets
+			foreach (XElement variableSetEl in element.Elements("variableSet"))
+			{
+				vars.ApplyVariableSet(variableSetEl, csvMetadata);
+			}
+
 			name = element.GetRequiredAttribute<string>(vars, "name");
 			title = element.GetRequiredAttribute<string>(vars, "title");
 			bStripEvents = element.GetSafeAttribute<bool>(vars, "stripEvents", true);
diff --git a/Engine/Source/Programs/CSVTools/PerfReportTool/ReportXml.cs b/Engine/Source/Programs/CSVTools/PerfReportTool/ReportXml.cs
--- a/Engine/Source/Programs/CSVTools/PerfReportTool/ReportXml.cs
+++ b/Engine/Source/Programs/CSVTools/PerfReportTool/ReportXml.cs
@@ -8,7 +8,6 @@ using System.Xml.Linq;
 using CSVStats;
 using PerfSummaries;
 using System.Globalization;
-using System.Runtime.InteropServices;
 
 namespace PerfReportTool
 {
@@ -147,6 +146,8 @@ namespace PerfReportTool
 				throw new Exception("No reporttypes element found in report XML " + reportTypeXmlFilename);
 			}
 
+			globalVariableSetElement = rootElement.Element("globalVariableSet");
+
 			// Read the graph XML
 			string graphsXMLFilename;
 			if (graphXMLFilenameIn.Length > 0)
@@ -328,6 +329,11 @@ namespace PerfReportTool
 					vars.SetVariable("meta." + key, metadataDict[key]);
 				}
 			}
+			if (globalVariableSetElement != null)
+			{
+				// Apply the global variable set
+				vars.ApplyVariableSet(globalVariableSetElement, csvFile.metadata);
+			}
 
 			ReportTypeInfo reportTypeInfo = null;
 			if (reportType == "")
@@ -340,7 +346,7 @@ namespace PerfReportTool
 
 					if (bReportTypeSupportsAutodetect && IsReportTypeXMLCompatibleWithStats(element, csvFile.dummyCsvStats, vars))
 					{
-						reportTypeInfo = new ReportTypeInfo(element, sharedSummaries, baseXmlDirectory, vars);
+						reportTypeInfo = new ReportTypeInfo(element, sharedSummaries, baseXmlDirectory, vars, csvFile.metadata);
 						break;
 					}
 
@@ -361,7 +367,7 @@ namespace PerfReportTool
 						throw new Exception("Default report type " + defaultReportTypeName + " was not compatible with CSV " + csvFile.filename);
 					}
 					Console.Out.WriteLine("Falling back to default report type: " + defaultReportTypeName);
-					reportTypeInfo = new ReportTypeInfo(defaultReportTypeElement, sharedSummaries, baseXmlDirectory, vars);
+					reportTypeInfo = new ReportTypeInfo(defaultReportTypeElement, sharedSummaries, baseXmlDirectory, vars, csvFile.metadata);
 				}
 				else if (reportTypeInfo == null)
 				{
@@ -394,7 +400,7 @@ namespace PerfReportTool
 						throw new Exception("Report type " + reportType + " is not compatible with CSV " + csvFile.filename);
 					}
 				}
-				reportTypeInfo = new ReportTypeInfo(foundReportTypeElement, sharedSummaries, baseXmlDirectory, vars);
+				reportTypeInfo = new ReportTypeInfo(foundReportTypeElement, sharedSummaries, baseXmlDirectory, vars, csvFile.metadata);
 			}
 
 			// Load the graphs
@@ -534,6 +540,7 @@ namespace PerfReportTool
 		XElement rootElement;
 		XElement graphGroupsElement;
 		XElement summaryTablesElement;
+		XElement globalVariableSetElement;
 		string defaultReportTypeName;
 		Dictionary<string, XElement> sharedSummaries;
 		Dictionary<string, GraphSettings> graphs;
@@ -591,6 +598,28 @@ namespace PerfReportTool
 			}
 			return attributeValue;
 		}
+
+		public void ApplyVariableSet(XElement variableSetElement, CsvMetadata csvMetadata)
+		{
+			string metadataQuery = variableSetElement.GetSafeAttribute<string>(this, "metadataQuery");
+
+			if ( metadataQuery == null || ( csvMetadata != null && CsvStats.DoesMetadataMatchFilter(csvMetadata, metadataQuery) ) )
+			{
+				// We match, so apply all variables and then apply all recursive variablesets
+				foreach (XElement variable in variableSetElement.Elements("var"))
+				{
+					string name = variable.FirstAttribute.Name.ToString();
+					string value = ResolveVariables(variable.FirstAttribute.Value);
+					SetVariable(name, value);
+				}
+
+				foreach (XElement childVariableSet in variableSetElement.Elements("variableSet"))
+				{
+					ApplyVariableSet(childVariableSet, csvMetadata);
+				}
+			}
+		}
+
 		Dictionary<string, string> vars = new Dictionary<string, string>();
 	}
 	
