PerfReportTool 4.221.0 : Add support for metadata queries and metadata variables (ie {$meta.NAME} ) in derived metadata mappings. This makes these declarations more powerful and reduces duplication
Metadata queries can also be blank if necessary if we want the values to always be set.
SourceName/DestName are still supported for back-compat if no query is provided
Note: There is no globalvariable support here. Solving this is not worth the complexity, performance hit and potential ordering and correctness implications
#tests Passes Test cases with identical results

diff --git a/Engine/Source/Programs/CSVTools/PerfReportTool/CsvFileCache..cs b/Engine/Source/Programs/CSVTools/PerfReportTool/CsvFileCache..cs
--- a/Engine/Source/Programs/CSVTools/PerfReportTool/CsvFileCache..cs
+++ b/Engine/Source/Programs/CSVTools/PerfReportTool/CsvFileCache..cs
@@ -110,11 +110,19 @@ namespace PerfReportTool
 			}
 			filename = inFilename;
 
-			// Apply metadata mappings if we have metadata. Note: we skip for PRCs, since these already have mappings cached
-			if (dummyCsvStats != null && dummyCsvStats.metaData != null)
+			// Setup initial metadata variable mappings and derived metadata (which can read from the initial variableMappings). 
+			// Note: we skip for PRCs (dummyCsvStats==null), since these already have mappings cached
+			if (dummyCsvStats != null)
 			{
-				metadata = dummyCsvStats.metaData;
-				derivedMetadataMappings.ApplyMapping(metadata);
+				xmlVariableMappings = new XmlVariableMappings();
+
+				// Note: We could apply the global variable set here, but that is slow and leads to complexity and correctness issues
+				if (dummyCsvStats.metaData != null)
+				{
+					metadata = dummyCsvStats.metaData;
+					xmlVariableMappings.SetMetadataVariables(metadata);
+					derivedMetadataMappings.ApplyMapping(metadata, xmlVariableMappings);
+				}
 			}
 		}
 
@@ -123,6 +131,12 @@ namespace PerfReportTool
 			if (fileType == FileType.BinaryCsv)
 			{
 				finalCsv = CsvStats.ReadBinFile(filename);
+				// If we already have metadata then use it. This avoids the need to remap it multiple times
+				if (metadata != null)
+				{
+					finalCsv.metaData = metadata;
+				}
+
 			}
 			else
 			{
@@ -142,6 +156,13 @@ namespace PerfReportTool
 			if (fileType == FileType.TextCsv)
 			{
 				finalCsv = CsvStats.ReadCSVFromLines(textCsvLines, null);
+
+				// If we already have metadata then use it. This avoids the need to remap it multiple times
+				if ( metadata != null )
+				{
+					finalCsv.metaData = metadata;
+				}
+
 				return finalCsv;
 			}
 			return null;
@@ -204,6 +225,7 @@ namespace PerfReportTool
 		public CsvMetadata metadata;
 		public SummaryTableRowData cachedSummaryTableRowData;
 		public ReportTypeInfo reportTypeInfo;
+		public XmlVariableMappings xmlVariableMappings;
 		FileType fileType;
 	};
 
diff --git a/Engine/Source/Programs/CSVTools/PerfReportTool/PerfReportTool.cs b/Engine/Source/Programs/CSVTools/PerfReportTool/PerfReportTool.cs
--- a/Engine/Source/Programs/CSVTools/PerfReportTool/PerfReportTool.cs
+++ b/Engine/Source/Programs/CSVTools/PerfReportTool/PerfReportTool.cs
@@ -23,7 +23,7 @@ namespace PerfReportTool
     class Version
     {
 		// Format: Major.Minor.Bugfix
-        private static string VersionString = "4.220.0";
+        private static string VersionString = "4.221.0";
 
         public static string Get() { return VersionString; }
     };
@@ -1207,8 +1207,6 @@ namespace PerfReportTool
 		CsvStats ReadCsvStats(CachedCsvFile csvFile, int minX, int maxX)
 		{
 			CsvStats csvStats = csvFile.GetFinalCsv();
-			reportXML.ApplyDerivedMetadata(csvStats.metaData);
-
 			if (csvStats.metaData == null)
 			{
 				csvStats.metaData = new CsvMetadata();
diff --git a/Engine/Source/Programs/CSVTools/PerfReportTool/ReportXml.cs b/Engine/Source/Programs/CSVTools/PerfReportTool/ReportXml.cs
--- a/Engine/Source/Programs/CSVTools/PerfReportTool/ReportXml.cs
+++ b/Engine/Source/Programs/CSVTools/PerfReportTool/ReportXml.cs
@@ -19,17 +19,45 @@ namespace PerfReportTool
 
 	class DerivedMetadataEntry
 	{
-		public DerivedMetadataEntry(string inSourceName, string inSourceValue, string inDestName, string inDestValue)
+		public DerivedMetadataEntry(XElement derivedMetadataEntry)
+		{
+			metadataQuery = derivedMetadataEntry.GetSafeAttribute<string>("metadataQuery");
+			if (metadataQuery == null)
+			{
+				// Back-compat: support sourceName/destName if metadataQuery isn't provided
+				string sourceName = derivedMetadataEntry.GetSafeAttribute<string>("sourceName");
+				if (sourceName != null)
+				{
+					string sourceValue = derivedMetadataEntry.GetRequiredAttribute<string>("sourceValue");
+					if (sourceValue != null)
+					{
+						metadataQuery = sourceName + "=" + sourceValue;
+					}
+				}
+			}
+			destName = derivedMetadataEntry.GetRequiredAttribute<string>("destName");
+			destValue = derivedMetadataEntry.GetRequiredAttribute<string>("destValue");
+			bHasVariables = (destName.Contains("${") || destValue.Contains("${") || (metadataQuery != null && metadataQuery.Contains("${")));
+		}
+		private DerivedMetadataEntry()
 		{
-			sourceName = inSourceName;
-			sourceValue = inSourceValue;
-			destName = inDestName;
-			destValue = inDestValue;
 		}
-		public string sourceName;
-		public string sourceValue;
+		public DerivedMetadataEntry ApplyVariableMappings(XmlVariableMappings vars)
+		{
+			if (vars == null || !bHasVariables)
+			{
+				return this;
+			}
+			DerivedMetadataEntry newEntry = new DerivedMetadataEntry();
+			newEntry.metadataQuery = metadataQuery != null ? vars.ResolveVariables(metadataQuery) : null;
+			newEntry.destName = vars.ResolveVariables(destName);
+			newEntry.destValue = vars.ResolveVariables(destValue);
+			return newEntry;
+		}
+		public string metadataQuery;
 		public string destName;
 		public string destValue;
+		public bool bHasVariables;
 	};
 
 	class DerivedMetadataMappings
@@ -38,20 +66,28 @@ namespace PerfReportTool
 		{
 			entries = new List<DerivedMetadataEntry>();
 		}
-		public void ApplyMapping(CsvMetadata csvMetadata)
+		public void ApplyMapping(CsvMetadata csvMetadata, XmlVariableMappings vars)
 		{
 			if (csvMetadata != null)
 			{
-				foreach (DerivedMetadataEntry entry in entries)
+				List<KeyValuePair<string, string>> valuesToAdd = new List<KeyValuePair<string, string>>();
+				foreach (DerivedMetadataEntry rawEntry in entries)
 				{
-					if (csvMetadata.Values.ContainsKey(entry.sourceName.ToLowerInvariant()))
+					DerivedMetadataEntry entry = rawEntry.ApplyVariableMappings(vars);
+
+					// Only override if the key is not already in the CSV metadata
+					if (!csvMetadata.Values.ContainsKey(entry.destName.ToLowerInvariant()))
 					{
-						if (csvMetadata.Values[entry.sourceName].ToLowerInvariant() == entry.sourceValue.ToLowerInvariant())
+						if (entry.metadataQuery == null || CsvStats.DoesMetadataMatchFilter(csvMetadata, entry.metadataQuery)) 
 						{
-							csvMetadata.Values.Add(entry.destName.ToLowerInvariant(), entry.destValue);
+							valuesToAdd.Add(new KeyValuePair<string, string>(entry.destName.ToLowerInvariant(), entry.destValue));
 						}
 					}
 				}
+				foreach (KeyValuePair<string,string> pair in valuesToAdd)
+				{
+					csvMetadata.Values[pair.Key] = pair.Value;
+				}
 			}
 		}
 		public List<DerivedMetadataEntry> entries;
@@ -233,15 +269,7 @@ namespace PerfReportTool
 			{
 				foreach (XElement mapping in derivedMetadataMappingsElement.Elements("mapping"))
 				{
-					string sourceName = mapping.GetSafeAttribute<string>("sourceName");
-					string sourceValue = mapping.GetSafeAttribute<string>("sourceValue");
-					string destName = mapping.GetSafeAttribute<string>("destName");
-					string destValue = mapping.GetSafeAttribute<string>("destValue");
-					if (sourceName == null || sourceValue == null || destName == null || destValue == null)
-					{
-						throw new Exception("Derivedmetadata mapping is missing a required attribute!\nRequired attributes: sourceName, sourceValue, destName, destValue.\nXML: " + mapping.ToString());
-					}
-					derivedMetadataMappings.entries.Add(new DerivedMetadataEntry(sourceName, sourceValue, destName, destValue));
+					derivedMetadataMappings.entries.Add(new DerivedMetadataEntry(mapping));
 				}
 			}
 
@@ -319,22 +347,12 @@ namespace PerfReportTool
 
 		public ReportTypeInfo GetReportTypeInfo(string reportType, CachedCsvFile csvFile, bool bBulkMode, bool forceReportType)
 		{
-			// Setup the variable mappings
-			XmlVariableMappings vars = new XmlVariableMappings();
-			if (csvFile.metadata != null)
-			{
-				Dictionary<string, string> metadataDict = csvFile.metadata.Values;
-				foreach (string key in metadataDict.Keys)
-				{
-					vars.SetVariable("meta." + key, metadataDict[key]);
-				}
-			}
+			XmlVariableMappings vars = csvFile.xmlVariableMappings;
+			// Apply the global variable set
 			if (globalVariableSetElement != null)
 			{
-				// Apply the global variable set
 				vars.ApplyVariableSet(globalVariableSetElement, csvFile.metadata);
 			}
-
 			ReportTypeInfo reportTypeInfo = null;
 			if (reportType == "")
 			{
@@ -524,11 +542,6 @@ namespace PerfReportTool
 			return csvEventsToStrip;
 		}
 
-		public void ApplyDerivedMetadata(CsvMetadata csvMetadata)
-		{
-			derivedMetadataMappings.ApplyMapping(csvMetadata);
-		}
-
 		public List<string> GetSummaryTableNames()
 		{
 			return summaryTables.Keys.ToList();
@@ -560,6 +573,15 @@ namespace PerfReportTool
 			vars[Name] = Value;
 		}
 
+		public void SetMetadataVariables(CsvMetadata csvMetadata)
+		{
+			Dictionary<string, string> metadataDict = csvMetadata.Values;
+			foreach (string key in metadataDict.Keys)
+			{
+				SetVariable("meta." + key, metadataDict[key]);
+			}
+		}
+
 		public string ResolveVariables(string attributeValue)
 		{
 			// Remap all variables found in the attribute name
