Add Contains and Get methods for SummaryTableRowData

diff --git a/Engine/Source/Programs/CSVTools/PerfReportTool/SummaryTableData.cs b/Engine/Source/Programs/CSVTools/PerfReportTool/SummaryTableData.cs
--- a/Engine/Source/Programs/CSVTools/PerfReportTool/SummaryTableData.cs
+++ b/Engine/Source/Programs/CSVTools/PerfReportTool/SummaryTableData.cs
@@ -439,6 +439,19 @@ namespace PerfSummaries
 			return (int)dict["framecount"].numericValue;
 		}
 
+		public bool Contains(string key)
+		{
+			return dict.ContainsKey(key);
+		}
+		public SummaryTableElement Get(string key)
+		{
+			if ( dict.TryGetValue(key, out SummaryTableElement elementOut) )
+			{
+				return elementOut;
+			}
+			return null;
+		}
+
 		public void RemoveSafe(string name)
 		{
 			string key = name.ToLower();
