PerfReportTool 4.222.0 : support dereferencing variables with array indices for strings separated with commas. Also add validation for variable names.

diff --git a/Engine/Source/Programs/CSVTools/PerfReportTool/PerfReportTool.cs b/Engine/Source/Programs/CSVTools/PerfReportTool/PerfReportTool.cs
--- a/Engine/Source/Programs/CSVTools/PerfReportTool/PerfReportTool.cs
+++ b/Engine/Source/Programs/CSVTools/PerfReportTool/PerfReportTool.cs
@@ -23,7 +23,7 @@ namespace PerfReportTool
     class Version
     {
 		// Format: Major.Minor.Bugfix
-        private static string VersionString = "4.221.0";
+        private static string VersionString = "4.222.0";
 
         public static string Get() { return VersionString; }
     };
diff --git a/Engine/Source/Programs/CSVTools/PerfReportTool/ReportXml.cs b/Engine/Source/Programs/CSVTools/PerfReportTool/ReportXml.cs
--- a/Engine/Source/Programs/CSVTools/PerfReportTool/ReportXml.cs
+++ b/Engine/Source/Programs/CSVTools/PerfReportTool/ReportXml.cs
@@ -517,7 +517,7 @@ namespace PerfReportTool
 				}
 			}
 
-			//Console.Out.WriteLine("Autodetected report type: " + reportTypeName);
+			Console.Out.WriteLine("Autodetected report type: " + reportTypeName);
 
 			return true;
 		}
@@ -568,8 +568,17 @@ namespace PerfReportTool
 
 	class XmlVariableMappings
 	{
-		public void SetVariable(string Name, string Value)
+		public void SetVariable(string Name, string Value, bool bValidate = true)
 		{
+			if (bValidate)
+			{
+				// Check for legal characters in the name
+				if (!Name.All(x => char.IsLetterOrDigit(x) || x == '.'))
+				{
+					throw new Exception("Invalid variable name: " + Name);
+				}
+			}
+
 			vars[Name] = Value;
 		}
 
@@ -578,7 +587,7 @@ namespace PerfReportTool
 			Dictionary<string, string> metadataDict = csvMetadata.Values;
 			foreach (string key in metadataDict.Keys)
 			{
-				SetVariable("meta." + key, metadataDict[key]);
+				SetVariable("meta." + key, metadataDict[key], false);
 			}
 		}
 
@@ -604,18 +613,55 @@ namespace PerfReportTool
 				{
 					break;
 				}
-				string VariableName = attributeValue.Substring(VarStartIndex+2, VarEndIndex - VarStartIndex-2);
+
+				// Advance StringPos
+				StringPos = VarEndIndex;
+
+				string FullVariableName = attributeValue.Substring(VarStartIndex+2, VarEndIndex - VarStartIndex-2);
+				string VariableName = FullVariableName;
+				int ArrayIndex = -1;
+
+				// Check for an array index
+				int OpenBracketIndex = VariableName.IndexOf('[');
+				if (OpenBracketIndex != -1)
+				{
+					if (FullVariableName.EndsWith("]"))
+					{
+						string ArrayIndexStr = VariableName.Substring(OpenBracketIndex + 1, VariableName.Length - 2 - OpenBracketIndex);
+						if (!int.TryParse(ArrayIndexStr, out ArrayIndex))
+						{
+							ArrayIndex = -1;
+						}
+					}
+					if (ArrayIndex < 0)
+					{
+						Console.WriteLine("[Warning] Failed to resolve variable $(" + FullVariableName + "}. Can't read array index");
+						continue;
+					}
+					VariableName = FullVariableName.Substring(0, OpenBracketIndex);
+				}
+
 
 				// Replace the variable if found
 				if (vars.TryGetValue(VariableName, out string VariableValue))
 				{
 					attributeValue = attributeValue.Substring(0,VarStartIndex) + VariableValue + attributeValue.Substring(VarEndIndex+1);
+					if (ArrayIndex >= 0)
+					{
+						string[] elements = attributeValue.Split(",");
+						if (ArrayIndex >= elements.Length)
+						{
+							Console.WriteLine("[Warning] Failed to resolve variable $(" + FullVariableName + "}. Array index out of range!");
+							continue;
+						}
+						attributeValue = elements[ArrayIndex];
+					}
+					// Adjust stringPos to take into account the replace
 					StringPos = VarStartIndex + VariableValue.Length;
 				}
 				else
 				{
 					Console.WriteLine("[Warning] Failed to resolve variable $(" + VariableName + "}");
-					StringPos = VarEndIndex;
 				}
 			}
 			return attributeValue;
