PerfReportTool 4.223.0. Budget and variable improvements
- Multiplier support for variables. Add multiplier="N" attribute to a variableset or variable
- FPSCharts : Compute frametime color thresholds automatically if MVP thresholds are specified and frametime is not
- Support empty color thresholds (useful if we want to clear via a variable)
- Add support for dumping the variables to the log for each CSV via -dumpVariables param

diff --git a/Engine/Source/Programs/CSVTools/PerfReportTool/PerfReportTool.cs b/Engine/Source/Programs/CSVTools/PerfReportTool/PerfReportTool.cs
--- a/Engine/Source/Programs/CSVTools/PerfReportTool/PerfReportTool.cs
+++ b/Engine/Source/Programs/CSVTools/PerfReportTool/PerfReportTool.cs
@@ -23,7 +23,7 @@ namespace PerfReportTool
     class Version
     {
 		// Format: Major.Minor.Bugfix
-        private static string VersionString = "4.222.0";
+        private static string VersionString = "4.223.0";
 
         public static string Get() { return VersionString; }
     };
@@ -100,6 +100,7 @@ namespace PerfReportTool
 			"     Not available in bulk mode.\n" +
 			"  -noSmooth : disable smoothing on all graphs\n" +
 			"  -listSummaryTables: lists available summary tables from the current report XML\n" +
+			"  -dumpVariables: dumps all variables to the log for each CSV\n" +
 			"\n" +
 			"Performance args:\n" +
 			"  -perfLog : output performance logging information\n" +
@@ -586,6 +587,12 @@ namespace PerfReportTool
 								GenerateReport(cachedCsvFile, outputDir, bBulkMode, rowData, bBatchedGraphs, writeDetailedReports, bReadAllStats || bWriteToSummaryTableCache, cachedCsvFile.reportTypeInfo, csvDir);
 								perfLog.LogTiming("  GenerateReport");
 
+								if (GetBoolArg("dumpVariables") && cachedCsvFile.xmlVariableMappings != null)
+								{
+									Console.WriteLine("\nDumping variables for " + cachedCsvFile.filename + "\n");
+									cachedCsvFile.xmlVariableMappings.DumpToLog();
+								}
+
 								if (rowData != null && bWriteToSummaryTableCache)
 								{
 									if (rowData.WriteToCache(summaryTableCacheDir, cachedCsvFile.summaryTableCacheId))
diff --git a/Engine/Source/Programs/CSVTools/PerfReportTool/ReportXml.cs b/Engine/Source/Programs/CSVTools/PerfReportTool/ReportXml.cs
--- a/Engine/Source/Programs/CSVTools/PerfReportTool/ReportXml.cs
+++ b/Engine/Source/Programs/CSVTools/PerfReportTool/ReportXml.cs
@@ -591,6 +591,20 @@ namespace PerfReportTool
 			}
 		}
 
+		public void DumpToLog(bool bIncludeMetadata = false)
+		{
+			string[] keys = vars.Keys.ToArray();
+			Array.Sort(keys);
+			foreach (string key in keys)
+			{
+				if (!key.StartsWith("meta.") || bIncludeMetadata)
+				{
+					string value = vars[key].Replace(", ",",").Replace(",",", "); // Ensure padding for arrays
+					Console.WriteLine(key.PadRight(50) + value);
+				}
+			}
+		}
+
 		public string ResolveVariables(string attributeValue)
 		{
 			// Remap all variables found in the attribute name
@@ -635,7 +649,7 @@ namespace PerfReportTool
 					}
 					if (ArrayIndex < 0)
 					{
-						Console.WriteLine("[Warning] Failed to resolve variable $(" + FullVariableName + "}. Can't read array index");
+						Console.WriteLine("[Warning] Failed to resolve variable ${" + FullVariableName + "}. Can't read array index");
 						continue;
 					}
 					VariableName = FullVariableName.Substring(0, OpenBracketIndex);
@@ -651,7 +665,7 @@ namespace PerfReportTool
 						string[] elements = attributeValue.Split(",");
 						if (ArrayIndex >= elements.Length)
 						{
-							Console.WriteLine("[Warning] Failed to resolve variable $(" + FullVariableName + "}. Array index out of range!");
+							Console.WriteLine("[Warning] Failed to resolve variable ${" + FullVariableName + "}. Array index out of range!");
 							continue;
 						}
 						attributeValue = elements[ArrayIndex];
@@ -661,29 +675,52 @@ namespace PerfReportTool
 				}
 				else
 				{
-					Console.WriteLine("[Warning] Failed to resolve variable $(" + VariableName + "}");
+					Console.WriteLine("[Warning] Failed to resolve variable ${" + VariableName + "}");
 				}
 			}
 			return attributeValue;
 		}
 
-		public void ApplyVariableSet(XElement variableSetElement, CsvMetadata csvMetadata)
+		public void ApplyVariableSet(XElement variableSetElement, CsvMetadata csvMetadata, double parentMultiplier=1.0)
 		{
 			string metadataQuery = variableSetElement.GetSafeAttribute<string>(this, "metadataQuery");
+			double multiplier = variableSetElement.GetSafeAttribute<double>(this, "multiplier", 1.0) * parentMultiplier;
 
 			if ( metadataQuery == null || ( csvMetadata != null && CsvStats.DoesMetadataMatchFilter(csvMetadata, metadataQuery) ) )
 			{
-				// We match, so apply all variables and then apply all recursive variablesets
-				foreach (XElement variable in variableSetElement.Elements("var"))
-				{
-					string name = variable.FirstAttribute.Name.ToString();
-					string value = ResolveVariables(variable.FirstAttribute.Value);
-					SetVariable(name, value);
-				}
+				// We match, so apply all variables and recursive variablesets in order
+				foreach (XElement child in variableSetElement.Elements())
+				{ 
+					if (child.Name == "var")
+					{
+						string name = child.FirstAttribute.Name.ToString();
+						string value = ResolveVariables(child.FirstAttribute.Value);
+						double variableMultiplier = child.GetSafeAttribute<double>(this, "multiplier", 1.0) * multiplier;
 
-				foreach (XElement childVariableSet in variableSetElement.Elements("variableSet"))
-				{
-					ApplyVariableSet(childVariableSet, csvMetadata);
+						// Apply the multplier if there is one before setting the variable
+						if (variableMultiplier != 1.0)
+						{
+							string[] arrayValues = value.Split(',');
+							List<string> finalValues = new List<string>();
+							foreach (string elementValue in arrayValues)
+							{
+								if (!double.TryParse(elementValue, out double doubleVal))
+								{
+									break;
+								}
+								finalValues.Add((variableMultiplier * doubleVal).ToString());
+							}
+							if (finalValues.Count == arrayValues.Length)
+							{
+								value = string.Join(',', finalValues);
+							}
+						}
+						SetVariable(name, value);
+					}
+					else if (child.Name == "variableSet")
+					{
+						ApplyVariableSet(child, csvMetadata);
+					}
 				}
 			}
 		}
diff --git a/Engine/Source/Programs/CSVTools/PerfReportTool/Summaries/FPSChartSummary.cs b/Engine/Source/Programs/CSVTools/PerfReportTool/Summaries/FPSChartSummary.cs
--- a/Engine/Source/Programs/CSVTools/PerfReportTool/Summaries/FPSChartSummary.cs
+++ b/Engine/Source/Programs/CSVTools/PerfReportTool/Summaries/FPSChartSummary.cs
@@ -157,6 +157,26 @@ namespace PerfSummaries
 			public bool isAverageValue;
 		};
 
+		private ColourThresholdList ComputeFrameTimeColorThresholdsFromMVP( int fps )
+		{
+			ColourThresholdList mvpColorThresholdList = GetStatColourThresholdList("MVP"+ fps.ToString());
+			ColourThresholdList colorThresholdListOut = null;
+			if (mvpColorThresholdList != null && mvpColorThresholdList.Count == 4)
+			{
+				colorThresholdListOut = new ColourThresholdList();
+				double idealFrameTime = 1000.0 / (double)fps;
+				for (int i = 0; i < 4; i++)
+				{
+					// MVP = (1/idealFrameTime - 1/avgFrameTime) * idealFrameTime
+					// avgFrameTime = idealFrameTime / (1 - MVP)
+					ThresholdInfo mvpThreshold = mvpColorThresholdList.Thresholds[i];
+					double frameTimeThreshold = idealFrameTime / (1.0 - mvpThreshold.value / 100.0);
+					colorThresholdListOut.Thresholds.Add(new ThresholdInfo(frameTimeThreshold, mvpThreshold.colour));
+				}
+			}
+			return colorThresholdListOut;
+		}
+
 		public override void WriteSummaryData(System.IO.StreamWriter htmlFile, CsvStats csvStats, CsvStats csvStatsUnstripped, bool bWriteSummaryCsv, SummaryTableRowData rowData, string htmlFileName)
 		{
 			System.IO.StreamWriter statsCsvFile = null;
@@ -205,7 +225,6 @@ namespace PerfSummaries
 					}
 				}
 
-				float value = 0.0f;
 				string ValueType = " Avg";
 				bool bIsAvg = false;
 				if (!csvStats.Stats.ContainsKey(baseStatName.ToLower()))
@@ -215,6 +234,7 @@ namespace PerfSummaries
 
 				bool bUnstripped = statAttributes.Contains("unstripped");
 				StatSamples stat = bUnstripped ? csvStatsUnstripped.Stats[baseStatName.ToLower()] : csvStats.Stats[baseStatName.ToLower()];
+				float value;
 				if (statAttributes.Contains("min"))
 				{
 					value = stat.ComputeMinValue();
@@ -236,7 +256,15 @@ namespace PerfSummaries
 				{
 					detailStr = "all frames";
 				}
-				Columns.Add(new ColumnInfo(baseStatName + ValueType, value, GetStatColourThresholdList(statName), detailStr, bIsAvg));
+
+				ColourThresholdList colorThresholdList = GetStatColourThresholdList(statName);
+
+				// If the frametime color thresholds are not specified then compute them based on MVP
+				if (!bIgnoreMVP && colorThresholdList == null && baseStatName.ToLower() == "frametime" && fps > 0)
+				{
+					colorThresholdList = ComputeFrameTimeColorThresholdsFromMVP(fps);
+				}
+				Columns.Add(new ColumnInfo(baseStatName + ValueType, value, colorThresholdList, detailStr, bIsAvg));
 			}
 
 			// Output summary table row data
diff --git a/Engine/Source/Programs/CSVTools/PerfReportTool/Summary.cs b/Engine/Source/Programs/CSVTools/PerfReportTool/Summary.cs
--- a/Engine/Source/Programs/CSVTools/PerfReportTool/Summary.cs
+++ b/Engine/Source/Programs/CSVTools/PerfReportTool/Summary.cs
@@ -120,12 +120,15 @@ namespace PerfSummaries
                     }
                     string statName = child.GetRequiredAttribute<string>(vars, "stat");
 					string hitchThresholdsStr = child.GetValue(vars);
+					if (hitchThresholdsStr == "")
+					{
+						continue;
+					}
                     string[] hitchThresholdsStrList = hitchThresholdsStr.Split(',');
 					ColourThresholdList HitchThresholds = new ColourThresholdList();
 					for (int i = 0; i < hitchThresholdsStrList.Length; i++)
                     {
 						string hitchThresholdStr = hitchThresholdsStrList[i];
-						double thresholdValue = 0.0;
 						string hitchThresholdNumStr = hitchThresholdStr;
 						Colour thresholdColour = null;
 
@@ -140,7 +143,7 @@ namespace PerfSummaries
 								thresholdColour = new Colour(colourString);
 							}
 						}
-						thresholdValue = Convert.ToDouble(hitchThresholdNumStr, System.Globalization.CultureInfo.InvariantCulture);
+						double thresholdValue = Convert.ToDouble(hitchThresholdNumStr, System.Globalization.CultureInfo.InvariantCulture);
 
 						HitchThresholds.Add(new ThresholdInfo(thresholdValue, thresholdColour));
                     }
