PerfReportTool 4.224.0: Adjust default colorization to make greens and reds lighter and more readable
- Add Variable support for hitch thresholds

diff --git a/Engine/Source/Programs/CSVTools/PerfReportTool/PerfReportTool.cs b/Engine/Source/Programs/CSVTools/PerfReportTool/PerfReportTool.cs
--- a/Engine/Source/Programs/CSVTools/PerfReportTool/PerfReportTool.cs
+++ b/Engine/Source/Programs/CSVTools/PerfReportTool/PerfReportTool.cs
@@ -23,7 +23,7 @@ namespace PerfReportTool
     class Version
     {
 		// Format: Major.Minor.Bugfix
-        private static string VersionString = "4.223.0";
+        private static string VersionString = "4.224.0";
 
         public static string Get() { return VersionString; }
     };
diff --git a/Engine/Source/Programs/CSVTools/PerfReportTool/Summaries/HitchSummary.cs b/Engine/Source/Programs/CSVTools/PerfReportTool/Summaries/HitchSummary.cs
--- a/Engine/Source/Programs/CSVTools/PerfReportTool/Summaries/HitchSummary.cs
+++ b/Engine/Source/Programs/CSVTools/PerfReportTool/Summaries/HitchSummary.cs
@@ -16,7 +16,7 @@ namespace PerfSummaries
 		{
 			ReadStatsFromXML(element, vars);
 
-			string[] hitchThresholds = element.Element("hitchThresholds").Value.Split(',');
+			string[] hitchThresholds = element.Element("hitchThresholds").GetValue(vars).Split(',');
 			HitchThresholds = new double[hitchThresholds.Length];
 			for (int i = 0; i < hitchThresholds.Length; i++)
 			{
diff --git a/Engine/Source/Programs/CSVTools/PerfReportTool/SummaryTypes.cs b/Engine/Source/Programs/CSVTools/PerfReportTool/SummaryTypes.cs
--- a/Engine/Source/Programs/CSVTools/PerfReportTool/SummaryTypes.cs
+++ b/Engine/Source/Programs/CSVTools/PerfReportTool/SummaryTypes.cs
@@ -166,7 +166,7 @@ namespace PerfSummaries
 		public static Colour Red = new Colour(1.0f, 0.4f, 0.4f);
 		public static Colour Orange = new Colour(1.0f, 0.7f, 0.4f, 1.0f);
 		public static Colour Yellow = new Colour(1.0f, 1.0f, 0.5f);
-		public static Colour Green = new Colour(0.4f, 0.82f, 0.45f);
+		public static Colour Green = new Colour(0.53f, 0.83f, 0.53f);
 
 		public float r, g, b;
 		public float alpha;
