PerfReportTool 4.224.1: Derived metadata populates variables
-add dumpVariablesAll param to dump all variables including metadata

diff --git a/Engine/Source/Programs/CSVTools/PerfReportTool/PerfReportTool.cs b/Engine/Source/Programs/CSVTools/PerfReportTool/PerfReportTool.cs
--- a/Engine/Source/Programs/CSVTools/PerfReportTool/PerfReportTool.cs
+++ b/Engine/Source/Programs/CSVTools/PerfReportTool/PerfReportTool.cs
@@ -23,7 +23,7 @@ namespace PerfReportTool
     class Version
     {
 		// Format: Major.Minor.Bugfix
-        private static string VersionString = "4.224.0";
+        private static string VersionString = "4.224.1";
 
         public static string Get() { return VersionString; }
     };
@@ -100,7 +100,7 @@ namespace PerfReportTool
 			"     Not available in bulk mode.\n" +
 			"  -noSmooth : disable smoothing on all graphs\n" +
 			"  -listSummaryTables: lists available summary tables from the current report XML\n" +
-			"  -dumpVariables: dumps all variables to the log for each CSV\n" +
+			"  -dumpVariables | -dumpVariablesAll : dumps variables to the log for each CSV (all includes metadata)\n" +
 			"\n" +
 			"Performance args:\n" +
 			"  -perfLog : output performance logging information\n" +
@@ -587,10 +587,10 @@ namespace PerfReportTool
 								GenerateReport(cachedCsvFile, outputDir, bBulkMode, rowData, bBatchedGraphs, writeDetailedReports, bReadAllStats || bWriteToSummaryTableCache, cachedCsvFile.reportTypeInfo, csvDir);
 								perfLog.LogTiming("  GenerateReport");
 
-								if (GetBoolArg("dumpVariables") && cachedCsvFile.xmlVariableMappings != null)
+								if ( ( GetBoolArg("dumpVariables") || GetBoolArg("dumpVariablesAll") ) && cachedCsvFile.xmlVariableMappings != null)
 								{
 									Console.WriteLine("\nDumping variables for " + cachedCsvFile.filename + "\n");
-									cachedCsvFile.xmlVariableMappings.DumpToLog();
+									cachedCsvFile.xmlVariableMappings.DumpToLog(GetBoolArg("dumpVariablesAll"));
 								}
 
 								if (rowData != null && bWriteToSummaryTableCache)
diff --git a/Engine/Source/Programs/CSVTools/PerfReportTool/ReportXml.cs b/Engine/Source/Programs/CSVTools/PerfReportTool/ReportXml.cs
--- a/Engine/Source/Programs/CSVTools/PerfReportTool/ReportXml.cs
+++ b/Engine/Source/Programs/CSVTools/PerfReportTool/ReportXml.cs
@@ -80,7 +80,11 @@ namespace PerfReportTool
 					{
 						if (entry.metadataQuery == null || CsvStats.DoesMetadataMatchFilter(csvMetadata, entry.metadataQuery)) 
 						{
-							valuesToAdd.Add(new KeyValuePair<string, string>(entry.destName.ToLowerInvariant(), entry.destValue));
+							string key = entry.destName.ToLowerInvariant();
+							valuesToAdd.Add(new KeyValuePair<string, string>(key, entry.destValue));
+
+							// Add the derived metadata value to variables
+							vars.SetVariable("meta." + key, entry.destValue);
 						}
 					}
 				}
