PerfReportTool 4.225.0: Add support for inline graphs, so we can define graphs fully in a ReportType. This allows us to reference variables in the graph definition directly, e.g for declaring a stat string or setting a maxY which corresponds to the target FPS

diff --git a/Engine/Source/Programs/CSVTools/PerfReportTool/PerfReportTool.cs b/Engine/Source/Programs/CSVTools/PerfReportTool/PerfReportTool.cs
--- a/Engine/Source/Programs/CSVTools/PerfReportTool/PerfReportTool.cs
+++ b/Engine/Source/Programs/CSVTools/PerfReportTool/PerfReportTool.cs
@@ -23,7 +23,7 @@ namespace PerfReportTool
     class Version
     {
 		// Format: Major.Minor.Bugfix
-        private static string VersionString = "4.224.1";
+        private static string VersionString = "4.225.0";
 
         public static string Get() { return VersionString; }
     };
@@ -1009,6 +1009,11 @@ namespace PerfReportTool
 						svgFilenames.Add(svgFilename);
 						continue;
 					}
+
+					if (graph.settings.statString.isSet == false)
+					{
+						throw new Exception("Graph " + graph.title + " has no <statString> element");
+					}
 					bool bFoundStat = false;
 					foreach (string statString in graph.settings.statString.value.Split(','))
 					{
diff --git a/Engine/Source/Programs/CSVTools/PerfReportTool/ReportGraph.cs b/Engine/Source/Programs/CSVTools/PerfReportTool/ReportGraph.cs
--- a/Engine/Source/Programs/CSVTools/PerfReportTool/ReportGraph.cs
+++ b/Engine/Source/Programs/CSVTools/PerfReportTool/ReportGraph.cs
@@ -97,8 +97,20 @@ namespace PerfReportTool
 			budget = new Optional<double>(element, "budget", vars);
 			inSummary = element.GetSafeAttribute<bool>(vars, "inSummary", false);
 			isExternal = element.GetSafeAttribute<bool>(vars, "external", false);
-
+			isInline = element.GetSafeAttribute<bool>(vars, "inline", false);
+			parent = element.GetSafeAttribute<string>(vars, "parent");
 			minFilterStatValue = new Optional<double>(element, "minFilterStatValue", vars);
+
+			if (!isInline && parent != null)
+			{
+				throw new Exception("Parent can only be specified for inline graphs (inline='1'): " + element.ToString());
+			}
+
+			if (isInline)
+			{
+				// If this is an inline graph then just load the settings directly
+				settings = new GraphSettings(element, vars);
+			}
 		}
 		public string title;
 		public Optional<double> budget;
@@ -106,42 +118,45 @@ namespace PerfReportTool
 		public bool isExternal;
 		public Optional<double> minFilterStatValue;
 		public GraphSettings settings;
+
+		public bool isInline;
+		public string parent;
 	};
 
 	class GraphSettings
 	{
-		public GraphSettings(XElement element)
+		public GraphSettings(XElement element, XmlVariableMappings vars = null)
 		{
-			smooth = new Optional<bool>(element, "smooth");
-			thickness = new Optional<double>(element, "thickness");
-			miny = new Optional<double>(element, "miny");
-			maxy = new Optional<double>(element, "maxy");
-			maxAutoMaxY = new Optional<double>(element, "maxAutoMaxY");
-			threshold = new Optional<double>(element, "threshold");
-			averageThreshold = new Optional<double>(element, "averageThreshold");
-			minFilterStatValue = new Optional<double>(element, "minFilterStatValue");
-			minFilterStatName = new OptionalString(element, "minFilterStatName");
-			smoothKernelPercent = new Optional<double>(element, "smoothKernelPercent");
-			smoothKernelSize = new Optional<double>(element, "smoothKernelSize");
-			compression = new Optional<double>(element, "compression");
-			width = new Optional<int>(element, "width");
-			height = new Optional<int>(element, "height");
-			stacked = new Optional<bool>(element, "stacked");
-			showAverages = new Optional<bool>(element, "showAverages");
-			filterOutZeros = new Optional<bool>(element, "filterOutZeros");
-			maxHierarchyDepth = new Optional<int>(element, "maxHierarchyDepth");
-			hideStatPrefix = new OptionalString(element, "hideStatPrefix");
-			mainStat = new OptionalString(element, "mainStat");
-			showEvents = new OptionalString(element, "showEvents");
-			requiresDetailedStats = new Optional<bool>(element, "requiresDetailedStats");
-			ignoreStats = new OptionalString(element, "ignoreStats");
-
-			statString = new OptionalString(element, "statString", true);
-			//additionalArgs = new OptionalString(element, "additionalArgs", true);
-			statMultiplier = new	(element, "statMultiplier");
-			legendAverageThreshold = new Optional<double>(element, "legendAverageThreshold");
-			snapToPeaks = new Optional<bool>(element, "snapToPeaks");
-			lineDecimalPlaces = new Optional<int>(element, "lineDecimalPlaces");
+			smooth = new Optional<bool>(element, "smooth", vars);
+			thickness = new Optional<double>(element, "thickness", vars);
+			miny = new Optional<double>(element, "miny", vars);
+			maxy = new Optional<double>(element, "maxy", vars);
+			maxAutoMaxY = new Optional<double>(element, "maxAutoMaxY", vars);
+			threshold = new Optional<double>(element, "threshold", vars);
+			averageThreshold = new Optional<double>(element, "averageThreshold", vars);
+			minFilterStatValue = new Optional<double>(element, "minFilterStatValue", vars);
+			minFilterStatName = new OptionalString(element, "minFilterStatName", false, vars);
+			smoothKernelPercent = new Optional<double>(element, "smoothKernelPercent", vars);
+			smoothKernelSize = new Optional<double>(element, "smoothKernelSize", vars);
+			compression = new Optional<double>(element, "compression", vars);
+			width = new Optional<int>(element, "width", vars);
+			height = new Optional<int>(element, "height", vars);
+			stacked = new Optional<bool>(element, "stacked", vars);
+			showAverages = new Optional<bool>(element, "showAverages", vars);
+			filterOutZeros = new Optional<bool>(element, "filterOutZeros", vars);
+			maxHierarchyDepth = new Optional<int>(element, "maxHierarchyDepth", vars);
+			hideStatPrefix = new OptionalString(element, "hideStatPrefix", false, vars);
+			mainStat = new OptionalString(element, "mainStat", false, vars);
+			showEvents = new OptionalString(element, "showEvents", false, vars);
+			requiresDetailedStats = new Optional<bool>(element, "requiresDetailedStats", vars);
+			ignoreStats = new OptionalString(element, "ignoreStats", false, vars);
+
+			statString = new OptionalString(element, "statString", true, vars);
+			//additionalArgs = new OptionalString(element, "additionalArgs", true, vars);
+			statMultiplier = new	(element, "statMultiplier", vars);
+			legendAverageThreshold = new Optional<double>(element, "legendAverageThreshold", vars);
+			snapToPeaks = new Optional<bool>(element, "snapToPeaks", vars);
+			lineDecimalPlaces = new Optional<int>(element, "lineDecimalPlaces", vars);
 		}
 		public void InheritFrom(GraphSettings baseSettings)
 		{
diff --git a/Engine/Source/Programs/CSVTools/PerfReportTool/ReportXml.cs b/Engine/Source/Programs/CSVTools/PerfReportTool/ReportXml.cs
--- a/Engine/Source/Programs/CSVTools/PerfReportTool/ReportXml.cs
+++ b/Engine/Source/Programs/CSVTools/PerfReportTool/ReportXml.cs
@@ -221,24 +221,18 @@ namespace PerfReportTool
 			}
 
 			graphs = new Dictionary<string, GraphSettings>();
-			foreach (XElement graphGroupElement in graphGroupsElement.Elements())
+			foreach (XElement graphGroupElement in graphGroupsElement.Elements("graphGroup"))
 			{
-				if (graphGroupElement.Name == "graphGroup")
+				// Create the base settings
+				XElement settingsElement = graphGroupElement.Element("baseSettings");
+				GraphSettings groupSettings = new GraphSettings(settingsElement);
+				groupSettings.InheritFrom(baseSettings);
+				foreach (XElement graphElement in graphGroupElement.Elements("graph"))
 				{
-					// Create the base settings
-					XElement settingsElement = graphGroupElement.Element("baseSettings");
-					GraphSettings groupSettings = new GraphSettings(settingsElement);
-					groupSettings.InheritFrom(baseSettings);
-					foreach (XElement graphElement in graphGroupElement.Elements())
-					{
-						if (graphElement.Name == "graph")
-						{
-							string title = graphElement.GetRequiredAttribute<string>("title").ToLower();
-							GraphSettings graphSettings = new GraphSettings(graphElement);
-							graphSettings.InheritFrom(groupSettings);
-							graphs.Add(title, graphSettings);
-						}
-					}
+					string title = graphElement.GetRequiredAttribute<string>("title").ToLower();
+					GraphSettings graphSettings = new GraphSettings(graphElement);
+					graphSettings.InheritFrom(groupSettings);
+					graphs.Add(title, graphSettings);
 				}
 			}
 
@@ -428,14 +422,24 @@ namespace PerfReportTool
 			// Load the graphs
 			foreach (ReportGraph graph in reportTypeInfo.graphs)
 			{
-				string key = graph.title.ToLower();
-				if (graphs.ContainsKey(key))
+				if (graph.isInline)
 				{
-					graph.settings = graphs[key];
+					if (graph.parent != null)
+					{
+						GraphSettings parentSettings = null;
+						if (!graphs.TryGetValue(graph.parent.ToLower(), out parentSettings))
+						{
+							throw new Exception("Parent graph with title \"" + graph.parent + "\" was not found in graphs XML");
+						}
+						graph.settings.InheritFrom(parentSettings);
+					}
 				}
 				else
 				{
-					throw new Exception("Graph with title \"" + graph.title + "\" was not found in graphs XML");
+					if (!graphs.TryGetValue(graph.title.ToLower(), out graph.settings))
+					{
+						throw new Exception("Graph with title \"" + graph.title + "\" was not found in graphs XML");
+					}
 				}
 			}
 
