HistogramThresholds variables support + hook up to XML

diff --git a/Engine/Source/Programs/CSVTools/PerfReportTool/Summaries/HistogramSummary.cs b/Engine/Source/Programs/CSVTools/PerfReportTool/Summaries/HistogramSummary.cs
--- a/Engine/Source/Programs/CSVTools/PerfReportTool/Summaries/HistogramSummary.cs
+++ b/Engine/Source/Programs/CSVTools/PerfReportTool/Summaries/HistogramSummary.cs
@@ -17,7 +17,7 @@ namespace PerfSummaries
 
 			ColourThresholds = ReadColourThresholdsXML(element.Element("colourThresholds"), vars);
 
-			string[] histogramStrings = element.Element("histogramThresholds").Value.Split(',');
+			string[] histogramStrings = element.Element("histogramThresholds").GetValue(vars).Split(',');
 			HistogramThresholds = new double[histogramStrings.Length];
 			for (int i = 0; i < histogramStrings.Length; i++)
 			{
@@ -85,7 +85,7 @@ namespace PerfSummaries
 
 			foreach (float thresh in thresholds)
 			{
-				htmlFile.WriteLine("  <td><b>>" + thresh.ToString() + "ms</b></td>");
+				htmlFile.WriteLine("  <td><b>>" + thresh.ToString("0.0") + "ms</b></td>");
 			}
 			htmlFile.WriteLine("  </tr>");
 
