Changing CSV library to emit a more useful error when encountering an empty file.

diff --git a/Engine/Source/Programs/CSVTools/CsvStats/CsvStats.cs b/Engine/Source/Programs/CSVTools/CsvStats/CsvStats.cs
--- a/Engine/Source/Programs/CSVTools/CsvStats/CsvStats.cs
+++ b/Engine/Source/Programs/CSVTools/CsvStats/CsvStats.cs
@@ -554,6 +554,10 @@ namespace CSVStats
 
 			// Detect unicode
 			string line = reader.ReadLine();
+			if (line == null)
+			{
+				throw new Exception(filename + " is empty. Aborting further processing to avoid unexpected results.");
+			}
 
 			bool bIsUnicode = false;
 			for (int i = 0; i < line.Length - 1; i++)
