PerfReportTool: Optimize commandline parsing. Some args can be very long when response files are passed so use a stringbuilder for concatenation. Commandline parsing could take 40s for a summarytable with 3000 PRC entries before.

diff --git a/Engine/Source/Programs/CSVTools/CsvStats/CommandLineTool.cs b/Engine/Source/Programs/CSVTools/CsvStats/CommandLineTool.cs
--- a/Engine/Source/Programs/CSVTools/CsvStats/CommandLineTool.cs
+++ b/Engine/Source/Programs/CSVTools/CsvStats/CommandLineTool.cs
@@ -26,7 +26,7 @@ namespace CSVStats
 		{
 			List<string> args = new List<string>();
 			bool bInQuotes = false;
-			string currentArg = "";
+			StringBuilder currentArgSb= new StringBuilder();
 			for (int i=0; i< inCommandLine.Length; i++)
 			{
 				char c = inCommandLine[i];
@@ -38,29 +38,29 @@ namespace CSVStats
 				{
 					if (bInQuotes)
 					{
-						currentArg += c;
+						currentArgSb.Append(c);
 					}
 					else
 					{ 
 						if (Char.IsWhiteSpace(c))
 						{
-							if (currentArg.Length > 0)
+							if (currentArgSb.Length > 0)
 							{
-								args.Add(currentArg);
-								currentArg = "";
+								args.Add(currentArgSb.ToString());
+								currentArgSb.Clear();
 							}
 						}
 						else
 						{
-							currentArg += c;
+							currentArgSb.Append(c);
 						}
 
 					}
 				}
 			}
-			if (currentArg.Length > 0)
+			if (currentArgSb.Length > 0)
 			{
-				args.Add(currentArg);
+				args.Add(currentArgSb.ToString());
 			}
 			return args;
 		}
diff --git a/Engine/Source/Programs/CSVTools/PerfReportTool/PerfReportTool.cs b/Engine/Source/Programs/CSVTools/PerfReportTool/PerfReportTool.cs
--- a/Engine/Source/Programs/CSVTools/PerfReportTool/PerfReportTool.cs
+++ b/Engine/Source/Programs/CSVTools/PerfReportTool/PerfReportTool.cs
@@ -23,7 +23,7 @@ namespace PerfReportTool
     class Version
     {
 		// Format: Major.Minor.Bugfix
-        private static string VersionString = "4.226.0";
+        private static string VersionString = "4.226.1";
 
         public static string Get() { return VersionString; }
     };
