PerfReportTool 4.227.0
Make summaries and graphs collapsible. Summaries can start collapsed if specified in xml via the "collapsed" attribute.
Add a new HtmlSection helper class to assist with this.
We may later want to extend this functionality to add a top level hierarchy to the report, allowing us to have high level sections with multiple summaries in them.

diff --git a/Engine/Source/Programs/CSVTools/PerfReportTool/PerfReportTool.cs b/Engine/Source/Programs/CSVTools/PerfReportTool/PerfReportTool.cs
--- a/Engine/Source/Programs/CSVTools/PerfReportTool/PerfReportTool.cs
+++ b/Engine/Source/Programs/CSVTools/PerfReportTool/PerfReportTool.cs
@@ -23,7 +23,7 @@ namespace PerfReportTool
     class Version
     {
 		// Format: Major.Minor.Bugfix
-        private static string VersionString = "4.226.1";
+        private static string VersionString = "4.227.0";
 
         public static string Get() { return VersionString; }
     };
@@ -1323,15 +1323,46 @@ namespace PerfReportTool
 				}
 				htmlFile.WriteLine("    ]]>");
 				htmlFile.WriteLine("    <title>" + titleStr + "</title>");
+
+				// Scripting for collapsibles
+				htmlFile.WriteLine("    <script>");
+				htmlFile.WriteLine("        document.addEventListener('DOMContentLoaded', function (event) { setupCollapsibles(); })");
+				htmlFile.WriteLine("        function setupCollapsibles() {");
+
+				htmlFile.WriteLine("            var collapsibles = document.getElementsByClassName('collapsibleHeading');");
+				htmlFile.WriteLine("            var i;");
+				htmlFile.WriteLine("            for (i = 0; i < collapsibles.length; i++) {");
+				htmlFile.WriteLine("                collapsibles[i].addEventListener('click', function() {");
+				htmlFile.WriteLine("                this.classList.toggle('expanded');");
+				htmlFile.WriteLine("                this.nextElementSibling.classList.toggle('expanded');");
+				htmlFile.WriteLine("             	});");
+				htmlFile.WriteLine("            }");
+				htmlFile.WriteLine("        }");
+				htmlFile.WriteLine("    </script>");
+
+				// CSS
 				htmlFile.WriteLine("    <style type='text/css'>");
+
 				htmlFile.WriteLine("      table, th, td { border: 2px solid black; border-collapse: collapse; padding: 3px; vertical-align: top; font-family: 'Verdana', Times, serif; font-size: 12px;}");
 				htmlFile.WriteLine("      p {  font-family: 'Verdana', Times, serif; font-size: 12px }");
+				htmlFile.WriteLine("      ul {  font-family: 'Verdana', Times, serif; font-size: 14px }");
 				htmlFile.WriteLine("      h1 {  font-family: 'Verdana', Times, serif; font-size: 20px; padding-top:10px }");
-				htmlFile.WriteLine("      h2 {  font-family: 'Verdana', Times, serif; font-size: 18px; padding-top:20px }");
-				htmlFile.WriteLine("      h3 {  font-family: 'Verdana', Times, serif; font-size: 16px; padding-top:20px }");
+				htmlFile.WriteLine("      h2 {  font-family: 'Verdana', Times, serif; font-size: 18px; padding-top:5px; padding-bottom:0px; margin-block-end: 0.4em }");
+				htmlFile.WriteLine("      h3 {  font-family: 'Verdana', Times, serif; font-size: 16px; padding-top:5px }");
+
+				// Collapsibles
+				htmlFile.WriteLine("      .collapsibleHeading { background-color: #ffffff; cursor: pointer; width: fit-content; border: none; text-align: left; outline: none; display: flex; justify-content: flex-start; align-items: flex-end; flex-direction: row; flex-wrap: nowrap; }");
+				htmlFile.WriteLine("      .expanded { background-color: #fff; }");
+				htmlFile.WriteLine("      .collapsibleHeading:hover { background-color: #c4c4c4; }");
+				htmlFile.WriteLine("      .collapsibleHeading:after { content: '\u25B2'; font-size: 14px; color: #4d4d4d; font-variant: normal; float: right; margin-left: 5px;}");
+				htmlFile.WriteLine("      .collapsibleHeading.expanded:after { content: '\u25BC'; }");
+				htmlFile.WriteLine("      .collapsibleSection { display: grid; grid-template-rows: 0fr; transition: grid-template-rows 0.2s ease-out; }");
+				htmlFile.WriteLine("      .collapsibleSection.expanded { grid-template-rows: 1fr; }");
+				htmlFile.WriteLine("      .collapsibleSectionInner { overflow: auto; overflow-y: hidden; }");
+
 				htmlFile.WriteLine("    </style>");
 				htmlFile.WriteLine("  </head>");
-				htmlFile.WriteLine("  <body><font face='verdana'>");
+				htmlFile.WriteLine("  <body>");
 				htmlFile.WriteLine("  <h1>" + titleStr + "</h1>");
 
 				// show the range
@@ -1345,8 +1376,6 @@ namespace PerfReportTool
 					htmlFile.WriteLine(")</font>");
 				}
 
-				htmlFile.WriteLine("  <h2>Summary</h2>");
-
 				htmlFile.WriteLine("<table style='width:800'>");
 
 				if (reportTypeInfo.metadataToShowList != null)
@@ -1395,7 +1424,12 @@ namespace PerfReportTool
 			PeakSummary peakSummary = null;
 			foreach (Summary summary in summaries)
 			{
-				summary.WriteSummaryData(htmlFile, summary.useUnstrippedCsvStats ? csvStatsUnstripped : csvStats, csvStatsUnstripped, bWriteSummaryCsv, summaryRowData, htmlFilename);
+				HtmlSection htmlSection = summary.WriteSummaryData(htmlFile != null, summary.useUnstrippedCsvStats ? csvStatsUnstripped : csvStats, csvStatsUnstripped, bWriteSummaryCsv, summaryRowData, htmlFilename);
+				if (htmlSection != null)
+				{
+					htmlSection.WriteToFile(htmlFile);
+				}
+					
 				if (summary.GetType() == typeof(PeakSummary))
 				{
 					peakSummary = (PeakSummary)summary;
@@ -1454,12 +1488,12 @@ namespace PerfReportTool
 					ReportGraph graph = graphs[svgFileIndex];
 
 					string svgTitle = graph.title;
-					htmlFile.WriteLine("  <br><a name='" + StripSpaces(svgTitle) + "'></a> <h2>" + svgTitle + "</h2>");
+					HtmlSection htmlSection = new HtmlSection(svgTitle, false, StripSpaces(svgTitle));
 					if (graph.isExternal)
 					{
 						string outFilename = htmlFilename.Replace(".html", "_" + svgTitle.Replace(" ", "_") + ".svg");
 						File.Copy(svgFilename, outFilename, true);
-						htmlFile.WriteLine("<a href='" + outFilename + "'>" + svgTitle + " (external)</a>");
+						htmlSection.WriteLine("<a href='" + outFilename + "'>" + svgTitle + " (external)</a>");
 					}
 					else
 					{
@@ -1467,9 +1501,10 @@ namespace PerfReportTool
 						foreach (string line in svgLines)
 						{
 							string modLine = line.Replace("__MAKEUNIQUE__", "U_" + svgFileIndex.ToString());
-							htmlFile.WriteLine(modLine);
+							htmlSection.WriteLine(modLine);
 						}
 					}
+					htmlSection.WriteToFile(htmlFile);
 				}
 
 				if (GetBoolArg("noWatermarks"))
@@ -1480,7 +1515,6 @@ namespace PerfReportTool
 				{
 					htmlFile.WriteLine("<p style='font-size:8'>Created with PerfReportTool " + Version.Get() + "</p>");
 				}
-				htmlFile.WriteLine("  </font>");
 				htmlFile.WriteLine("  </body>");
 				htmlFile.WriteLine("</html>");
 				htmlFile.Close();
@@ -1544,7 +1578,12 @@ namespace PerfReportTool
 			PeakSummary peakSummary = null;
 			foreach (Summary summary in reportTypeInfo.summaries)
 			{
-				summary.WriteSummaryData(htmlFile, csvStats, csvStatsUnstripped, bWriteSummaryCsv, null, htmlFilename);
+				HtmlSection htmlSection = summary.WriteSummaryData(htmlFile != null, csvStats, csvStatsUnstripped, bWriteSummaryCsv, null, htmlFilename);
+				if (htmlSection != null)
+				{
+					htmlSection.WriteToFile(htmlFile);
+				}
+
 				if (summary.GetType() == typeof(PeakSummary))
 				{
 					peakSummary = (PeakSummary)summary;
diff --git a/Engine/Source/Programs/CSVTools/PerfReportTool/Summaries/BoundedStatValuesSummary.cs b/Engine/Source/Programs/CSVTools/PerfReportTool/Summaries/BoundedStatValuesSummary.cs
--- a/Engine/Source/Programs/CSVTools/PerfReportTool/Summaries/BoundedStatValuesSummary.cs
+++ b/Engine/Source/Programs/CSVTools/PerfReportTool/Summaries/BoundedStatValuesSummary.cs
@@ -81,8 +81,9 @@ namespace PerfSummaries
 
 		public override string GetName() { return "boundedstatvalues"; }
 
-		public override void WriteSummaryData(System.IO.StreamWriter htmlFile, CsvStats csvStats, CsvStats csvStatsUnstripped, bool bWriteSummaryCsv, SummaryTableRowData rowData, string htmlFileName)
+		public override HtmlSection WriteSummaryData(bool bWriteHtml, CsvStats csvStats, CsvStats csvStatsUnstripped, bool bWriteSummaryCsv, SummaryTableRowData rowData, string htmlFileName)
 		{
+			HtmlSection htmlSection = null;
 			int startFrame = -1;
 			int endFrame = int.MaxValue;
 
@@ -100,7 +101,7 @@ namespace PerfSummaries
 				if (startFrame == -1)
 				{
 					Console.WriteLine("BoundedStatValuesSummary: Begin event " + beginEvent + " was not found");
-					return;
+					return htmlSection;
 				}
 			}
 			if (endEvent != null)
@@ -119,13 +120,13 @@ namespace PerfSummaries
 				if (endFrame == int.MaxValue)
 				{
 					Console.WriteLine("BoundedStatValuesSummary: End event " + endEvent + " was not found");
-					return;
+					return htmlSection;
 				}
 			}
 			if (startFrame >= endFrame)
 			{
 				Console.WriteLine("Warning: BoundedStatValuesSummary: end event "+ endEvent + " appeared before the start event "+beginEvent);
-				return;
+				return htmlSection;
 			}
 			endFrame = Math.Min(endFrame, csvStats.SampleCount - 1);
 			startFrame = Math.Max(startFrame, 0);
@@ -155,7 +156,7 @@ namespace PerfSummaries
 			// Nothing to report, so bail out!
 			if (filteredColumns.Count == 0)
 			{
-				return;
+				return htmlSection;
 			}
 
 			// Process the column values
@@ -358,17 +359,18 @@ namespace PerfSummaries
 			}
 
 			// Output HTML
-			if (htmlFile != null)
+			if (bWriteHtml)
 			{
-				htmlFile.WriteLine("  <h2>" + title + "</h2>");
-				htmlFile.WriteLine("  <table border='0' style='width:1400'>");
-				htmlFile.WriteLine("  <tr>");
+				htmlSection = new HtmlSection(title, bStartCollapsed);
+
+				htmlSection.WriteLine("  <table border='0' style='width:1400'>");
+				htmlSection.WriteLine("  <tr>");
 				foreach (Column col in filteredColumns)
 				{
-					htmlFile.WriteLine("<th>" + col.name + "</th>");
+					htmlSection.WriteLine("<th>" + col.name + "</th>");
 				}
-				htmlFile.WriteLine("  </tr>");
-				htmlFile.WriteLine("  <tr>");
+				htmlSection.WriteLine("  </tr>");
+				htmlSection.WriteLine("  <tr>");
 				foreach (Column col in filteredColumns)
 				{
 					string bgcolor = "'#ffffff'";
@@ -376,10 +378,10 @@ namespace PerfSummaries
 					{
 						bgcolor = col.colourThresholdList.GetColourForValue(col.value);
 					}
-					htmlFile.WriteLine("<td bgcolor=" + bgcolor + ">" + col.value.ToString("0.00") + "</td>");
+					htmlSection.WriteLine("<td bgcolor=" + bgcolor + ">" + col.value.ToString("0.00") + "</td>");
 				}
-				htmlFile.WriteLine("  </tr>");
-				htmlFile.WriteLine("  </table>");
+				htmlSection.WriteLine("  </tr>");
+				htmlSection.WriteLine("  </table>");
 			}
 
 			// Output summary table row data
@@ -393,6 +395,7 @@ namespace PerfSummaries
 					}
 				}
 			}
+			return htmlSection;
 		}
 		public override void PostInit(ReportTypeInfo reportTypeInfo, CsvStats csvStats)
 		{
diff --git a/Engine/Source/Programs/CSVTools/PerfReportTool/Summaries/BucketSummary.cs b/Engine/Source/Programs/CSVTools/PerfReportTool/Summaries/BucketSummary.cs
--- a/Engine/Source/Programs/CSVTools/PerfReportTool/Summaries/BucketSummary.cs
+++ b/Engine/Source/Programs/CSVTools/PerfReportTool/Summaries/BucketSummary.cs
@@ -42,28 +42,29 @@ namespace PerfSummaries
 
 		public override string GetName() { return "bucketsummary"; }
 
-		public override void WriteSummaryData(System.IO.StreamWriter htmlFile, CsvStats csvStats, CsvStats csvStatsUnstripped, bool bWriteSummaryCsv, SummaryTableRowData metadata, string htmlFileName)
+		public override HtmlSection WriteSummaryData(bool bWriteHtml, CsvStats csvStats, CsvStats csvStatsUnstripped, bool bWriteSummaryCsv, SummaryTableRowData rowData, string htmlFileName)
 		{
 			// Only HTML reporting is supported (does not summary table row data)
-			if (htmlFile == null)
+			if (!bWriteHtml)
 			{
-				return;
+				return null;
 			}
 
 			if (Buckets.Length == 0)
 			{
-				return;
+				return null;
 			}
 
-			htmlFile.WriteLine("  <h2>" + Title + "</h2>");
-			htmlFile.WriteLine("  <table border='0' style='width:1000'>");
-			htmlFile.WriteLine("  <tr><td></td>");
+			HtmlSection htmlSection = new HtmlSection(Title, bStartCollapsed);
+
+			htmlSection.WriteLine("  <table border='0' style='width:1000'>");
+			htmlSection.WriteLine("  <tr><td></td>");
 
 			List<string> Hitches = new List<string>();
 
 			if (ReportBelowRange)
 			{
-				htmlFile.WriteLine("  <th> <" + Buckets[0].ToString("0") + "</b></td>");
+				htmlSection.WriteLine("  <th> <" + Buckets[0].ToString("0") + "</b></td>");
 			}
 
 			for (int i = 1; i < Buckets.Length; ++i)
@@ -71,15 +72,15 @@ namespace PerfSummaries
 				double Begin = Buckets[i - 1];
 				double End = Buckets[i];
 
-				htmlFile.WriteLine("  <th> [" + Begin.ToString("0") + ", " + End.ToString("0") + ")"+ "</b></td>");
+				htmlSection.WriteLine("  <th> [" + Begin.ToString("0") + ", " + End.ToString("0") + ")"+ "</b></td>");
 			}
 
 			if (ReportAboveRange)
 			{
-				htmlFile.WriteLine("  <th> >=" + Buckets.Last().ToString("0") + "</b></td>");
+				htmlSection.WriteLine("  <th> >=" + Buckets.Last().ToString("0") + "</b></td>");
 			}
 
-			htmlFile.WriteLine("  </tr>");
+			htmlSection.WriteLine("  </tr>");
 
 			foreach (string unitStat in stats)
 			{
@@ -91,7 +92,7 @@ namespace PerfSummaries
 				}
 
 				Hitches.Clear();
-				htmlFile.WriteLine("  <tr><td><b>" + StatToCheck + "</b></td>");
+				htmlSection.WriteLine("  <tr><td><b>" + StatToCheck + "</b></td>");
 				Hitches.Add(StatToCheck);
 
 				int[] BucketCounts = new int[Buckets.Length + 1];
@@ -131,13 +132,14 @@ namespace PerfSummaries
 					Color = Colour.Lerp(White, Color, Intensity);
 					string ColorString = Color.ToHTMLString();
 
-					htmlFile.WriteLine("  <td bgcolor=" + ColorString + ">" + Count.ToString("0") + "</td>");
+					htmlSection.WriteLine("  <td bgcolor=" + ColorString + ">" + Count.ToString("0") + "</td>");
 					Hitches.Add(Count.ToString("0"));
 				}
 
-				htmlFile.WriteLine("  </tr>");
+				htmlSection.WriteLine("  </tr>");
 			}
-			htmlFile.WriteLine("  </table>");
+			htmlSection.WriteLine("  </table>");
+			return htmlSection;
 		}
 
 		public double[] Buckets;
diff --git a/Engine/Source/Programs/CSVTools/PerfReportTool/Summaries/EventSummary.cs b/Engine/Source/Programs/CSVTools/PerfReportTool/Summaries/EventSummary.cs
--- a/Engine/Source/Programs/CSVTools/PerfReportTool/Summaries/EventSummary.cs
+++ b/Engine/Source/Programs/CSVTools/PerfReportTool/Summaries/EventSummary.cs
@@ -13,6 +13,7 @@ namespace PerfSummaries
 	{
 		public EventSummary(XElement element, XmlVariableMappings vars, string baseXmlDirectory)
 		{
+			ReadStatsFromXML(element, vars);
 			title = element.GetSafeAttribute(vars, "title", "Events");
 			summaryStatName = element.GetRequiredAttribute<string>(vars, "summaryStatName");
 			events = element.Element("events").Value.Split(',');
@@ -23,8 +24,9 @@ namespace PerfSummaries
 
 		public override string GetName() { return "event"; }
 
-		public override void WriteSummaryData(System.IO.StreamWriter htmlFile, CsvStats csvStats, CsvStats csvStatsUnstripped, bool bWriteSummaryCsv, SummaryTableRowData rowData, string htmlFileName)
+		public override HtmlSection WriteSummaryData(bool bWriteHtml, CsvStats csvStats, CsvStats csvStatsUnstripped, bool bWriteSummaryCsv, SummaryTableRowData rowData, string htmlFileName)
 		{
+			HtmlSection htmlSection = null;
 			Dictionary<string, int> eventCountsDict = new Dictionary<string, int>();
 			int eventCount = 0;
 			foreach (CsvEvent ev in csvStats.Events)
@@ -53,16 +55,16 @@ namespace PerfSummaries
 			}
 
 			// Output HTML
-			if (htmlFile != null && eventCountsDict.Count > 0)
+			if (bWriteHtml && eventCountsDict.Count > 0)
 			{
-				htmlFile.WriteLine("  <h2>" + title + "</h2>");
-				htmlFile.WriteLine("  <table border='0' style='width:1200'>");
-				htmlFile.WriteLine("  <tr><th>Name</th><th><b>Count</th></tr>");
+				htmlSection = new HtmlSection(title, bStartCollapsed);
+				htmlSection.WriteLine("  <table border='0' style='width:1200'>");
+				htmlSection.WriteLine("  <tr><th>Name</th><th><b>Count</th></tr>");
 				foreach (KeyValuePair<string, int> pair in eventCountsDict.ToList())
 				{
-					htmlFile.WriteLine("  <tr><td>" + pair.Key + "</td><td>" + pair.Value + "</td></tr>");
+					htmlSection.WriteLine("  <tr><td>" + pair.Key + "</td><td>" + pair.Value + "</td></tr>");
 				}
-				htmlFile.WriteLine("  </table>");
+				htmlSection.WriteLine("  </table>");
 			}
 
 			// Output summary table row data
@@ -81,6 +83,7 @@ namespace PerfSummaries
 				}
 				rowData.Add(SummaryTableElement.Type.SummaryTableMetric, summaryStatName, (double)eventCount, thresholdList);
 			}
+			return htmlSection;
 		}
 		public override void PostInit(ReportTypeInfo reportTypeInfo, CsvStats csvStats)
 		{
diff --git a/Engine/Source/Programs/CSVTools/PerfReportTool/Summaries/ExtraLinksSummary.cs b/Engine/Source/Programs/CSVTools/PerfReportTool/Summaries/ExtraLinksSummary.cs
--- a/Engine/Source/Programs/CSVTools/PerfReportTool/Summaries/ExtraLinksSummary.cs
+++ b/Engine/Source/Programs/CSVTools/PerfReportTool/Summaries/ExtraLinksSummary.cs
@@ -66,7 +66,8 @@ namespace PerfSummaries
         };
         public ExtraLinksSummary(XElement element, XmlVariableMappings vars, string baseXmlDirectory, bool bInLinkTemplates=false)
         {
-            title = "Links";
+			ReadStatsFromXML(element, vars);
+			title = "Links";
             if (element != null)
             {
                 title = element.GetSafeAttribute(vars, "title", title);
@@ -78,8 +79,9 @@ namespace PerfSummaries
 
 		public override string GetName() { return "extralinks"; }
 
-		public override void WriteSummaryData(System.IO.StreamWriter htmlFile, CsvStats csvStats, CsvStats csvStatsUnstripped, bool bWriteSummaryCsv, SummaryTableRowData rowData, string htmlFileName)
+		public override HtmlSection WriteSummaryData(bool bWriteHtml, CsvStats csvStats, CsvStats csvStatsUnstripped, bool bWriteSummaryCsv, SummaryTableRowData rowData, string htmlFileName)
 		{
+			HtmlSection htmlSection = null;
 			List<ExtraLink> links = new List<ExtraLink>();
 
             string csvId = null;
@@ -93,14 +95,14 @@ namespace PerfSummaries
 			if (csvFilename == null)
 			{
 				Console.WriteLine("Can't find CSV filename for ExtraLinks summary. Skipping");
-				return;
+				return null;
 			}
 
 			string linksFilename = csvFilename + ".links";
 			if (!File.Exists(linksFilename))
 			{
 				Console.WriteLine("Can't find file " + linksFilename + " for ExtraLinks summary. Skipping");
-				return;
+				return null;
 			}
 			string[] lines = File.ReadAllLines(linksFilename);
 			foreach (string line in lines)
@@ -109,19 +111,19 @@ namespace PerfSummaries
 			}
 			if (links.Count == 0)
 			{
-				return;
+				return null;
 			}
 
 			// Output HTML
-			if (htmlFile != null)
+			if (bWriteHtml)
 			{
-				htmlFile.WriteLine("  <h2>" + title + "</h2>");
-				htmlFile.WriteLine("  <ul>");
+				htmlSection = new HtmlSection(title, bStartCollapsed);
+				htmlSection.WriteLine("  <ul>");
 				foreach (ExtraLink link in links)
 				{
-					htmlFile.WriteLine("  <li>" + link.GetLinkString(false) + "</li>");
+					htmlSection.WriteLine("  <li>" + link.GetLinkString(false) + "</li>");
 				}
-				htmlFile.WriteLine("  </ul>");
+				htmlSection.WriteLine("  </ul>");
 			}
 
 			// Output summary row data
@@ -132,6 +134,8 @@ namespace PerfSummaries
 					rowData.Add(SummaryTableElement.Type.SummaryTableMetric, link.ReportText, link.GetLinkString(true), null);
 				}
 			}
+
+			return htmlSection;
 		}
 		public override void PostInit(ReportTypeInfo reportTypeInfo, CsvStats csvStats)
 		{
diff --git a/Engine/Source/Programs/CSVTools/PerfReportTool/Summaries/FPSChartSummary.cs b/Engine/Source/Programs/CSVTools/PerfReportTool/Summaries/FPSChartSummary.cs
--- a/Engine/Source/Programs/CSVTools/PerfReportTool/Summaries/FPSChartSummary.cs
+++ b/Engine/Source/Programs/CSVTools/PerfReportTool/Summaries/FPSChartSummary.cs
@@ -177,8 +177,10 @@ namespace PerfSummaries
 			return colorThresholdListOut;
 		}
 
-		public override void WriteSummaryData(System.IO.StreamWriter htmlFile, CsvStats csvStats, CsvStats csvStatsUnstripped, bool bWriteSummaryCsv, SummaryTableRowData rowData, string htmlFileName)
+		public override HtmlSection WriteSummaryData(bool bWriteHtml, CsvStats csvStats, CsvStats csvStatsUnstripped, bool bWriteSummaryCsv, SummaryTableRowData rowData, string htmlFileName)
 		{
+			HtmlSection htmlSection = null;
+
 			System.IO.StreamWriter statsCsvFile = null;
 			if (bWriteSummaryCsv)
 			{
@@ -298,8 +300,9 @@ namespace PerfSummaries
 			}
 
 			// Output HTML
-			if (htmlFile != null)
+			if (bWriteHtml)
 			{
+				htmlSection = new HtmlSection("FPSChart", bStartCollapsed);
 				string HeaderRow = "";
 				string ValueRow = "";
 				HeaderRow += "<th>Section Name</th>";
@@ -320,10 +323,9 @@ namespace PerfSummaries
 					HeaderRow += "<th>" + TableUtil.FormatStatName(columnName) + "</th>";
 					ValueRow += "<td bgcolor=" + column.Color + ">" + column.Value.ToString("0.00") + "</td>";
 				}
-				htmlFile.WriteLine("  <h2>FPSChart</h2>");
-				htmlFile.WriteLine("<table border='0' style='width:400'>");
-				htmlFile.WriteLine("  <tr>" + HeaderRow + "</tr>");
-				htmlFile.WriteLine("  <tr>" + ValueRow + "</tr>");
+				htmlSection.WriteLine("<table border='0' style='width:400'>");
+				htmlSection.WriteLine("  <tr>" + HeaderRow + "</tr>");
+				htmlSection.WriteLine("  <tr>" + ValueRow + "</tr>");
 			}
 
 			// Output CSV
@@ -442,7 +444,7 @@ namespace PerfSummaries
 					}
 
 					// Output HTML
-					if (htmlFile != null)
+					if (htmlSection != null)
 					{
 						string ValueRow = "";
 						ValueRow += "<td>" + CapRange.name + "</td>";
@@ -450,7 +452,7 @@ namespace PerfSummaries
 						{
 							ValueRow += "<td bgcolor=" + column.Color + ">" + column.Value.ToString("0.00") + "</td>";
 						}
-						htmlFile.WriteLine("  <tr>" + ValueRow + "</tr>");
+						htmlSection.WriteLine("  <tr>" + ValueRow + "</tr>");
 					}
 
 					// Output CSV
@@ -474,16 +476,17 @@ namespace PerfSummaries
 				}
 			}
 
-			if (htmlFile != null)
+			if (htmlSection != null)
 			{
-				htmlFile.WriteLine("</table>");
-				htmlFile.WriteLine("<p style='font-size:8'>Engine hitch metric: " + (bUseEngineHitchMetric ? "enabled" : "disabled") + "</p>");
+				htmlSection.WriteLine("</table>");
+				htmlSection.WriteLine("<p style='font-size:8'>Engine hitch metric: " + (bUseEngineHitchMetric ? "enabled" : "disabled") + "</p>");
 			}
 
 			if (statsCsvFile != null)
 			{
 				statsCsvFile.Close();
 			}
+			return htmlSection;
 		}
 		public override void PostInit(ReportTypeInfo reportTypeInfo, CsvStats csvStats)
 		{
diff --git a/Engine/Source/Programs/CSVTools/PerfReportTool/Summaries/HistogramSummary.cs b/Engine/Source/Programs/CSVTools/PerfReportTool/Summaries/HistogramSummary.cs
--- a/Engine/Source/Programs/CSVTools/PerfReportTool/Summaries/HistogramSummary.cs
+++ b/Engine/Source/Programs/CSVTools/PerfReportTool/Summaries/HistogramSummary.cs
@@ -40,54 +40,35 @@ namespace PerfSummaries
 
 		public override string GetName() { return "histogram"; }
 
-		public override void WriteSummaryData(System.IO.StreamWriter htmlFile, CsvStats csvStats, CsvStats csvStatsUnstripped, bool bWriteSummaryCsv, SummaryTableRowData metadata, string htmlFileName)
+		public override HtmlSection WriteSummaryData(bool bWriteHtml, CsvStats csvStats, CsvStats csvStatsUnstripped, bool bWriteSummaryCsv, SummaryTableRowData rowData, string htmlFileName)
 		{
 			// Only HTML reporting is supported (does not output summary table row data)
-			if (htmlFile == null)
+			if (!bWriteHtml)
 			{
-				return;
+				return null;
 			}
 
-			if (!bSuppressAveragesTable)
-			{
-				// Write the averages
-				htmlFile.WriteLine("  <h2>Stat unit averages</h2>");
-				htmlFile.WriteLine("  <table border='0' style='width:400'>");
-				htmlFile.WriteLine("  <tr><td></td><th>ms</b></th>");
-				foreach (string stat in stats)
-				{
-					string StatToCheck = stat.Split('(')[0];
-					if (!csvStats.Stats.ContainsKey(StatToCheck.ToLower()))
-					{
-						continue;
-					}
-					float val = csvStats.Stats[StatToCheck.ToLower()].average;
-					string colour = ColourThresholdList.GetThresholdColour(val, ColourThresholds[0], ColourThresholds[1], ColourThresholds[2], ColourThresholds[3]);
-					htmlFile.WriteLine("  <tr><td><b>" + StatToCheck + "</b></td><td bgcolor=" + colour + ">" + val.ToString("0.00") + "</td></tr>");
-				}
-				htmlFile.WriteLine("  </table>");
-			}
+			HtmlSection htmlSection = new HtmlSection("Histogram", bStartCollapsed);
 
 			// Histogram
 			double[] thresholds = HistogramThresholds;
-			htmlFile.WriteLine("  <h2>Histogram</h2>");
-			htmlFile.WriteLine("  <table border='0' style='width:800'>");
+			htmlSection.WriteLine("  <table border='0' style='width:800'>");
 
-			htmlFile.WriteLine("  <tr><td></td>");
+			htmlSection.WriteLine("  <tr><td></td>");
 
 			// Display the override stat budget first
 			bool HasBudgetOverrideStat = false;
 			if (BudgetOverrideStatName != null)
 			{
-				htmlFile.WriteLine("  <td><b><=" + BudgetOverrideStatBudget.ToString() + "ms</b></td>");
+				htmlSection.WriteLine("  <td><b><=" + BudgetOverrideStatBudget.ToString() + "ms</b></td>");
 				HasBudgetOverrideStat = true;
 			}
 
 			foreach (float thresh in thresholds)
 			{
-				htmlFile.WriteLine("  <td><b>>" + thresh.ToString("0.0") + "ms</b></td>");
+				htmlSection.WriteLine("  <td><b>>" + thresh.ToString("0.0") + "ms</b></td>");
 			}
-			htmlFile.WriteLine("  </tr>");
+			htmlSection.WriteLine("  </tr>");
 
 			foreach (string unitStat in stats)
 			{
@@ -97,7 +78,7 @@ namespace PerfSummaries
 					continue;
 				}
 
-				htmlFile.WriteLine("  <tr><td><b>" + StatToCheck + "</b></td>");
+				htmlSection.WriteLine("  <tr><td><b>" + StatToCheck + "</b></td>");
 				int thresholdIndex = 0;
 
 				// Display the render thread budget column (don't display the other stats)
@@ -108,11 +89,11 @@ namespace PerfSummaries
 						float RatioOfFramesOverBudget = 1.0f - csvStats.GetStat(StatToCheck.ToLower()).GetRatioOfFramesInBudget((float)BudgetOverrideStatBudget);
 						float pc = RatioOfFramesOverBudget * 100.0f;
 						string colour = ColourThresholdList.GetThresholdColour(pc, 50.0f, 35.0f, 20.0f, 0.0f);
-						htmlFile.WriteLine("  <td bgcolor=" + colour + ">" + pc.ToString("0.00") + "%</td>");
+						htmlSection.WriteLine("  <td bgcolor=" + colour + ">" + pc.ToString("0.00") + "%</td>");
 					}
 					else
 					{
-						htmlFile.WriteLine("  <td></td>");
+						htmlSection.WriteLine("  <td></td>");
 					}
 				}
 
@@ -122,12 +103,34 @@ namespace PerfSummaries
 					float RatioOfFramesOverBudget = 1.0f - csvStats.GetStat(StatToCheck.ToLower()).GetRatioOfFramesInBudget((float)threshold);
 					float pc = RatioOfFramesOverBudget * 100.0f;
 					string colour = ColourThresholdList.GetThresholdColour(pc, 50.0f, 35.0f, 20.0f, 0.0f);
-					htmlFile.WriteLine("  <td bgcolor=" + colour + ">" + pc.ToString("0.00") + "%</td>");
+					htmlSection.WriteLine("  <td bgcolor=" + colour + ">" + pc.ToString("0.00") + "%</td>");
 					thresholdIndex++;
 				}
-				htmlFile.WriteLine("  </tr>");
+				htmlSection.WriteLine("  </tr>");
+			}
+			htmlSection.WriteLine("  </table>");
+
+			if (!bSuppressAveragesTable)
+			{
+				// Write the averages
+				htmlSection.WriteLine("  <h3>Stat unit averages</h3>");
+				htmlSection.WriteLine("  <table border='0' style='width:400'>");
+				htmlSection.WriteLine("  <tr><td></td><th>ms</b></th>");
+				foreach (string stat in stats)
+				{
+					string StatToCheck = stat.Split('(')[0];
+					if (!csvStats.Stats.ContainsKey(StatToCheck.ToLower()))
+					{
+						continue;
+					}
+					float val = csvStats.Stats[StatToCheck.ToLower()].average;
+					string colour = ColourThresholdList.GetThresholdColour(val, ColourThresholds[0], ColourThresholds[1], ColourThresholds[2], ColourThresholds[3]);
+					htmlSection.WriteLine("  <tr><td><b>" + StatToCheck + "</b></td><td bgcolor=" + colour + ">" + val.ToString("0.00") + "</td></tr>");
+				}
+				htmlSection.WriteLine("  </table>");
 			}
-			htmlFile.WriteLine("  </table>");
+
+			return htmlSection;
 		}
 
 		public double[] ColourThresholds;
diff --git a/Engine/Source/Programs/CSVTools/PerfReportTool/Summaries/HitchSummary.cs b/Engine/Source/Programs/CSVTools/PerfReportTool/Summaries/HitchSummary.cs
--- a/Engine/Source/Programs/CSVTools/PerfReportTool/Summaries/HitchSummary.cs
+++ b/Engine/Source/Programs/CSVTools/PerfReportTool/Summaries/HitchSummary.cs
@@ -27,16 +27,18 @@ namespace PerfSummaries
 
 		public override string GetName() { return "hitches"; }
 
-		public override void WriteSummaryData(System.IO.StreamWriter htmlFile, CsvStats csvStats, CsvStats csvStatsUnstripped, bool bWriteSummaryCsv, SummaryTableRowData summaryTableRowData, string htmlFileName)
+		public override HtmlSection WriteSummaryData(bool bWriteHtml, CsvStats csvStats, CsvStats csvStatsUnstripped, bool bWriteSummaryCsv, SummaryTableRowData rowData, string htmlFileName)
 		{
+			HtmlSection htmlSection = null;
+
 			// if we have 20k frames in a typical flythrough then 20 frames would be red
 			float redThresholdFor50ms = (float)csvStats.SampleCount / 500.0f;
 
-			if (htmlFile != null)
+			if (bWriteHtml)
 			{
-				htmlFile.WriteLine("  <h2>Hitches</h2>");
-				htmlFile.WriteLine("  <table border='0' style='width:800'>");
-				htmlFile.WriteLine("  <tr><td></td>");
+				htmlSection = new HtmlSection("Hitches", bStartCollapsed);
+				htmlSection.WriteLine("  <table border='0' style='width:800'>");
+				htmlSection.WriteLine("  <tr><td></td>");
 
 				StreamWriter statsCsvFile = null;
 				if (bWriteSummaryCsv)
@@ -50,14 +52,14 @@ namespace PerfSummaries
 				Thresholds.Add("Hitch Size");
 				foreach (float thresh in HitchThresholds)
 				{
-					htmlFile.WriteLine("  <th> >" + thresh.ToString("0") + "ms</b></td>");
+					htmlSection.WriteLine("  <th> >" + thresh.ToString("0") + "ms</b></td>");
 					Thresholds.Add(thresh.ToString("0"));
 				}
 				if (statsCsvFile != null)
 				{
 					statsCsvFile.WriteLine(string.Join(",", Thresholds));
 				}
-				htmlFile.WriteLine("  </tr>");
+				htmlSection.WriteLine("  </tr>");
 
 				foreach (string unitStat in stats)
 				{
@@ -69,7 +71,7 @@ namespace PerfSummaries
 					}
 
 					Hitches.Clear();
-					htmlFile.WriteLine("  <tr><td><b>" + StatToCheck + "</b></td>");
+					htmlSection.WriteLine("  <tr><td><b>" + StatToCheck + "</b></td>");
 					Hitches.Add(StatToCheck);
 					int thresholdIndex = 0;
 
@@ -78,7 +80,7 @@ namespace PerfSummaries
 						float count = (float)statSample.GetCountOfFramesOverBudget(threshold);
 						float redThreshold = (redThresholdFor50ms * 50.0f) / threshold; // Adjust the colour threshold based on the current threshold
 						string colour = ColourThresholdList.GetThresholdColour(count, redThreshold, redThreshold * 0.66, redThreshold * 0.33, 0.0f);
-						htmlFile.WriteLine("  <td bgcolor=" + colour + ">" + count.ToString("0") + "</td>");
+						htmlSection.WriteLine("  <td bgcolor=" + colour + ">" + count.ToString("0") + "</td>");
 						Hitches.Add(count.ToString("0"));
 						thresholdIndex++;
 					}
@@ -87,18 +89,18 @@ namespace PerfSummaries
 						statsCsvFile.WriteLine(string.Join(",", Hitches));
 					}
 
-					htmlFile.WriteLine("  </tr>");
+					htmlSection.WriteLine("  </tr>");
 				}
 				if (statsCsvFile != null)
 				{
 					statsCsvFile.Close();
 				}
-				htmlFile.WriteLine("  </table>");
-				htmlFile.WriteLine("<p style='font-size:8'>Note: Simplified hitch metric. All frames over threshold are counted" + "</p>");
+				htmlSection.WriteLine("  </table>");
+				htmlSection.WriteLine("<p style='font-size:8'>Note: Simplified hitch metric. All frames over threshold are counted" + "</p>");
 			}
 
 			// Write out hitches to the summary for 100ms intervals up to 1000
-			if (summaryTableRowData != null)
+			if (rowData != null)
 			{
 				StatSamples frametimeStat = csvStats.GetStat("frametime");
 				if (frametimeStat != null)
@@ -111,10 +113,11 @@ namespace PerfSummaries
 						float redThreshold = (redThresholdFor50ms * 50.0f) / (float)threshold; // Adjust the colour threshold based on the current threshold
 						ColourThresholdList colorThresholdList=new ColourThresholdList(redThreshold,redThreshold*0.66,redThreshold*0.33,0.0);
 
-						summaryTableRowData.Add(SummaryTableElement.Type.SummaryTableMetric, "Hitches>" + threshold + "ms", (double)count, colorThresholdList);
+						rowData.Add(SummaryTableElement.Type.SummaryTableMetric, "Hitches>" + threshold + "ms", (double)count, colorThresholdList);
 					}
 				}
 			}
+			return htmlSection;
 		}
 		public double[] HitchThresholds;
 	};
diff --git a/Engine/Source/Programs/CSVTools/PerfReportTool/Summaries/MapOverlaySummary.cs b/Engine/Source/Programs/CSVTools/PerfReportTool/Summaries/MapOverlaySummary.cs
--- a/Engine/Source/Programs/CSVTools/PerfReportTool/Summaries/MapOverlaySummary.cs
+++ b/Engine/Source/Programs/CSVTools/PerfReportTool/Summaries/MapOverlaySummary.cs
@@ -143,11 +143,14 @@ namespace PerfSummaries
 			}
 		}
 
-		public override void WriteSummaryData(System.IO.StreamWriter htmlFile, CsvStats csvStats, CsvStats csvStatsUnstripped, bool bWriteSummaryCsv, SummaryTableRowData rowData, string htmlFileName)
+		public override HtmlSection WriteSummaryData(bool bWriteHtml, CsvStats csvStats, CsvStats csvStatsUnstripped, bool bWriteSummaryCsv, SummaryTableRowData rowData, string htmlFileName)
 		{
+			HtmlSection htmlSection = null;
+
 			// Output HTML
-			if (htmlFile != null)
+			if (bWriteHtml)
 			{
+				htmlSection = new HtmlSection(title, bStartCollapsed);
 				string outputDirectory = System.IO.Path.GetDirectoryName(System.IO.Path.GetFullPath(htmlFileName));
 				string outputMapFilename = System.IO.Path.Combine(outputDirectory, destImageFilename);
 
@@ -168,9 +171,8 @@ namespace PerfSummaries
 				}
 
 				// Check if the file exists in the output directory
-				htmlFile.WriteLine("  <h2>" + title + "</h2>");
-				htmlFile.WriteLine("<svg version='1.1' xmlns='http://www.w3.org/2000/svg' xmlns:xlink='http://www.w3.org/1999/xlink' width='" + imageWidth + "' height='" + imageHeight + "'>");
-				htmlFile.WriteLine("<image href='" + destImageFilename + "' width='" + imageWidth + "' height='" + imageHeight + "' />");
+				htmlSection.WriteLine("<svg version='1.1' xmlns='http://www.w3.org/2000/svg' xmlns:xlink='http://www.w3.org/1999/xlink' width='" + imageWidth + "' height='" + imageHeight + "'>");
+				htmlSection.WriteLine("<image href='" + destImageFilename + "' width='" + imageWidth + "' height='" + imageHeight + "' />");
 
 				// Draw the overlays
 				foreach (MapOverlay overlay in overlays)
@@ -215,7 +217,7 @@ namespace PerfSummaries
 					// Draw the lines
 					string currentLineColor = overlay.lineColor;
 					string lineStartTemplate = "<polyline style='fill:none;stroke-width:1.3;stroke:{LINECOLOUR}' points='";
-					htmlFile.Write(lineStartTemplate.Replace("{LINECOLOUR}", currentLineColor));
+					htmlSection.Write(lineStartTemplate.Replace("{LINECOLOUR}", currentLineColor));
 					float adjustedLineSplitDistanceThreshold = lineSplitDistanceThreshold * framesPerLineSegment;
 					float oldx = 0;
 					float oldy = 0;
@@ -249,20 +251,20 @@ namespace PerfSummaries
 						}
 						else
 						{
-							htmlFile.Write(lineCoordsStr);
+							htmlSection.Write(lineCoordsStr);
 						}
 
 						if (restartLineStrip)
 						{
-							htmlFile.WriteLine("'/>");
-							htmlFile.Write(lineStartTemplate.Replace("{LINECOLOUR}", currentLineColor));
-							htmlFile.Write(lineCoordsStr);
+							htmlSection.WriteLine("'/>");
+							htmlSection.Write(lineStartTemplate.Replace("{LINECOLOUR}", currentLineColor));
+							htmlSection.Write(lineCoordsStr);
 						}
 						oldx = x;
 						oldy = y;
 						lastFrameIndex = i;
 					}
-					htmlFile.WriteLine("'/>");
+					htmlSection.WriteLine("'/>");
 
 					// Plot the events 
 					float circleRadius = 3;
@@ -278,21 +280,22 @@ namespace PerfSummaries
 								float y = yStat.samples[ev.Frame];
 								int svgX = toSvgX(x, y);
 								int svgY = toSvgY(x, y);
-								htmlFile.Write("<circle cx='" + svgX + "' cy='" + svgY + "' r='" + circleRadius + "' fill='" + eventColourString + "' fill-opacity='1.0'/>");
-								htmlFile.WriteLine("<text x='" + (svgX + 5) + "' y='" + svgY + "' text-anchor='left' style='font-family: Verdana;fill: #ffffff; font-size: " + 9 + "px;'>" + eventText + "</text>");
+								htmlSection.Write("<circle cx='" + svgX + "' cy='" + svgY + "' r='" + circleRadius + "' fill='" + eventColourString + "' fill-opacity='1.0'/>");
+								htmlSection.WriteLine("<text x='" + (svgX + 5) + "' y='" + svgY + "' text-anchor='left' style='font-family: Verdana;fill: #ffffff; font-size: " + 9 + "px;'>" + eventText + "</text>");
 							}
 						}
 					}
 				}
 
-				//htmlFile.WriteLine("<text x='50%' y='" + (imageHeight * 0.05) + "' text-anchor='middle' style='font-family: Verdana;fill: #FFFFFF; stroke: #C0C0C0;  font-size: " + 20 + "px;'>" + title + "</text>");
-				htmlFile.WriteLine("</svg>");
+				//htmlSection.WriteLine("<text x='50%' y='" + (imageHeight * 0.05) + "' text-anchor='middle' style='font-family: Verdana;fill: #FFFFFF; stroke: #C0C0C0;  font-size: " + 20 + "px;'>" + title + "</text>");
+				htmlSection.WriteLine("</svg>");
 			}
 
 			// Output row data
 			if (rowData != null)
 			{
 			}
+			return htmlSection;
 		}
 		public override void PostInit(ReportTypeInfo reportTypeInfo, CsvStats csvStats)
 		{
diff --git a/Engine/Source/Programs/CSVTools/PerfReportTool/Summaries/PeakSummary.cs b/Engine/Source/Programs/CSVTools/PerfReportTool/Summaries/PeakSummary.cs
--- a/Engine/Source/Programs/CSVTools/PerfReportTool/Summaries/PeakSummary.cs
+++ b/Engine/Source/Programs/CSVTools/PerfReportTool/Summaries/PeakSummary.cs
@@ -84,16 +84,16 @@ namespace PerfSummaries
 
 		public override string GetName() { return "peak"; }
 
-		void WriteStatSection(StreamWriter htmlFile, CsvStats csvStats, PeakSummarySection section, StreamWriter LLMCsvData, SummaryTableRowData summaryTableRowData)
+		void WriteStatSection(HtmlSection htmlSection, CsvStats csvStats, PeakSummarySection section, StreamWriter LLMCsvData, SummaryTableRowData summaryTableRowData)
 		{
-			if (htmlFile != null)
+			if (htmlSection != null)
 			{
 				// Here we are deciding which title we have and write it to the file.
-				htmlFile.WriteLine("<h3>" + section.title + "</h3>");
-				htmlFile.WriteLine("  <table border='0' style='width:400'>");
+				htmlSection.WriteLine("<h3>" + section.title + "</h3>");
+				htmlSection.WriteLine("  <table border='0' style='width:400'>");
 
 				//Hard-coded start of the table.
-				htmlFile.WriteLine("    <tr><td style='width:200'></td><td style='width:75'><b>Average</b></td><td style='width:75'><b>Peak</b></td><td style='width:75'><b>Budget</b></td></tr>");
+				htmlSection.WriteLine("    <tr><td style='width:200'></td><td style='width:75'><b>Average</b></td><td style='width:75'><b>Peak</b></td><td style='width:75'><b>Budget</b></td></tr>");
 			}
 
 			foreach (PeakStatInfo statInfo in section.stats)
@@ -126,9 +126,9 @@ namespace PerfSummaries
 					averageColour = colorThresholdList.GetColourForValue(average);
 					budgetString = budget.ToString("0");
 				}
-				if (htmlFile != null)
+				if (htmlSection != null)
 				{
-					htmlFile.WriteLine("    <tr><td>" + statInfo.shortName + "</td><td bgcolor=" + averageColour + ">" + average.ToString("0") + "</td><td bgcolor=" + peakColour + ">" + peak.ToString("0") + "</td><td>" + budgetString + "</td></tr>");
+					htmlSection.WriteLine("    <tr><td>" + statInfo.shortName + "</td><td bgcolor=" + averageColour + ">" + average.ToString("0") + "</td><td bgcolor=" + peakColour + ">" + peak.ToString("0") + "</td><td>" + budgetString + "</td></tr>");
 				}
 
 				// Pass through color data as part of database-friendly stuff.
@@ -155,9 +155,9 @@ namespace PerfSummaries
 				}
 
 			}
-			if (htmlFile != null)
+			if (htmlSection != null)
 			{
-				htmlFile.WriteLine("  </table>");
+				htmlSection.WriteLine("  </table>");
 			}
 		}
 
@@ -173,8 +173,9 @@ namespace PerfSummaries
 			return null;
 		}
 
-		public override void WriteSummaryData(System.IO.StreamWriter htmlFile, CsvStats csvStats, CsvStats csvStatsUnstripped, bool bWriteSummaryCsv, SummaryTableRowData summaryTableRowData, string htmlFileName)
+		public override HtmlSection WriteSummaryData(bool bWriteHtml, CsvStats csvStats, CsvStats csvStatsUnstripped, bool bWriteSummaryCsv, SummaryTableRowData rowData, string htmlFileName)
 		{
+			HtmlSection htmlSection = null;
 			// Only HTML reporting is supported (does not output summary table row data)
 			StreamWriter LLMCsvData = null;
 			if (bWriteSummaryCsv)
@@ -196,19 +197,22 @@ namespace PerfSummaries
 				}
 			}
 
-			if (htmlFile != null)
+			if (bWriteHtml)
 			{
-				htmlFile.WriteLine("<h2>Peaks Summary</h2>");
+				htmlSection = new HtmlSection("Peaks Summary", bStartCollapsed);
 			}
+
 			foreach (PeakSummarySection section in peakSummarySections)
 			{
-				WriteStatSection(htmlFile, csvStats, section, LLMCsvData, summaryTableRowData);
+				WriteStatSection(htmlSection, csvStats, section, LLMCsvData, rowData);
 			}
 
 			if (LLMCsvData != null)
 			{
 				LLMCsvData.Close();
 			}
+
+			return htmlSection;
 		}
 
 
diff --git a/Engine/Source/Programs/CSVTools/PerfReportTool/Summaries/StatBudgetsSummary.cs b/Engine/Source/Programs/CSVTools/PerfReportTool/Summaries/StatBudgetsSummary.cs
--- a/Engine/Source/Programs/CSVTools/PerfReportTool/Summaries/StatBudgetsSummary.cs
+++ b/Engine/Source/Programs/CSVTools/PerfReportTool/Summaries/StatBudgetsSummary.cs
@@ -87,8 +87,9 @@ namespace PerfSummaries
 			return statName;
 		}
 
-		public override void WriteSummaryData(System.IO.StreamWriter htmlFile, CsvStats csvStats, CsvStats csvStatsUnstripped, bool bWriteSummaryCsv, SummaryTableRowData rowData, string htmlFileName)
+		public override HtmlSection WriteSummaryData(bool bWriteHtml, CsvStats csvStats, CsvStats csvStatsUnstripped, bool bWriteSummaryCsv, SummaryTableRowData rowData, string htmlFileName)
 		{
+			HtmlSection htmlSection = null;
 			// Find all referenced stats/metadata and sum them
 			double totalValue = 0.0;
 			List<StatInfo> statInfoList = new List<StatInfo>();
@@ -177,11 +178,11 @@ namespace PerfSummaries
 			}
 
 			// Output HTML
-			if (htmlFile != null)
+			if (bWriteHtml)
 			{
-				htmlFile.WriteLine("  <h2>" + title + "</h2>");
-				htmlFile.WriteLine("  <table border='0' style='width:400'>");
-				htmlFile.WriteLine("  <tr><th>Stat</th><th>Value (Avg)</th><th>Budget</th></tr>");
+				htmlSection = new HtmlSection(title, bStartCollapsed);
+				htmlSection.WriteLine("  <table border='0' style='width:400'>");
+				htmlSection.WriteLine("  <tr><th>Stat</th><th>Value (Avg)</th><th>Budget</th></tr>");
 				foreach (StatInfo statInfo in statInfoList)
 				{
 					if (showTotal == false && statInfo == totalStatInfo)
@@ -204,10 +205,11 @@ namespace PerfSummaries
 					{
 						displayName = StripStatPrefix(statInfo.name);
 					}
-					htmlFile.WriteLine("  <tr><td>" + displayName + "</td><td bgcolor=" + bgcolor + ">" + statInfo.averageValue.ToString("0.00") + "</td><td>" + budgetStr + "</td></tr>");
+					htmlSection.WriteLine("  <tr><td>" + displayName + "</td><td bgcolor=" + bgcolor + ">" + statInfo.averageValue.ToString("0.00") + "</td><td>" + budgetStr + "</td></tr>");
 				}
-				htmlFile.WriteLine("  </table>");
+				htmlSection.WriteLine("  </table>");
 			}
+			return htmlSection;
 		}
 
 		ColourThresholdList colorThresholdsTotal = null;
diff --git a/Engine/Source/Programs/CSVTools/PerfReportTool/Summary.cs b/Engine/Source/Programs/CSVTools/PerfReportTool/Summary.cs
--- a/Engine/Source/Programs/CSVTools/PerfReportTool/Summary.cs
+++ b/Engine/Source/Programs/CSVTools/PerfReportTool/Summary.cs
@@ -8,6 +8,7 @@ using CSVStats;
 using PerfReportTool;
 using System.Reflection;
 using System.Drawing;
+using System.Text;
 
 namespace PerfSummaries
 {
@@ -79,10 +80,12 @@ namespace PerfSummaries
             StatThresholds = new Dictionary<string, ColourThresholdList>();
 
         }
-        public virtual void WriteSummaryData(System.IO.StreamWriter htmlFile, CsvStats csvStats, CsvStats csvStatsUnstripped, bool bWriteSummaryCsv, SummaryTableRowData rowData, string htmlFileName)
-        { }
+		public virtual HtmlSection WriteSummaryData(bool bWriteHtml, CsvStats csvStats, CsvStats csvStatsUnstripped, bool bWriteSummaryCsv, SummaryTableRowData rowData, string htmlFileName)
+		{
+			return null;
+		}
 
-        public virtual void PostInit(ReportTypeInfo reportTypeInfo, CsvStats csvStats)
+		public virtual void PostInit(ReportTypeInfo reportTypeInfo, CsvStats csvStats)
         {
 			// Resolve wildcards and remove duplicates
 			stats = csvStats.GetStatNamesMatchingStringList(stats.ToArray());
@@ -93,6 +96,7 @@ namespace PerfSummaries
         public void ReadStatsFromXML(XElement element, XmlVariableMappings vars)
         {
 			useUnstrippedCsvStats = element.GetSafeAttribute<bool>(vars, "useUnstrippedCsvStats", false);
+			bStartCollapsed = element.GetSafeAttribute<bool>(vars, "collapsed", false);
 			XElement statsElement = element.Element("stats");
 			if (statsElement != null)
 			{
@@ -233,11 +237,90 @@ namespace PerfSummaries
         public List<string> stats;
         public Dictionary<string, ColourThresholdList> StatThresholds;
 		public bool useUnstrippedCsvStats;
+		public bool bStartCollapsed;
     };
 
 
 
+	class HtmlSection
+	{
+		public HtmlSection(string titleIn, bool bStartCollapsedIn, string anchorNameIn = null)
+		{
+			title = titleIn;
+			bStartCollapsed = bStartCollapsedIn;
+			anchorName = anchorNameIn; 
+		}
+
+		public void WriteLine(string text) 
+		{
+			FlushPendingLine();
+			lines.Add(text);
+		}
+		public void Write(string text)
+		{
+			pendingLine.Append(text);
+		}
+
+		public void WriteToFile( System.IO.StreamWriter htmlFile )
+		{
+			FlushPendingLine();
+			BeginHtmlSection(htmlFile);
+			foreach (string line in lines)
+			{
+				htmlFile.WriteLine(line);
+			}
+			EndHtmlSection(htmlFile);
+		}
 
+		private void FlushPendingLine()
+		{
+			if (pendingLine.Length > 0)
+			{
+				lines.Add(pendingLine.ToString());
+				pendingLine.Clear();
+			}
+		}
+
+		protected void BeginHtmlSection(System.IO.StreamWriter htmlFile)
+		{
+			if (htmlFile != null)
+			{
+				string headingClass = "collapsibleHeading";
+				string divClass = "collapsibleSection";
+
+				if (!bStartCollapsed)
+				{
+					headingClass += " expanded";
+					divClass += " expanded";
+				}
+
+				if (anchorName != null)
+				{
+					htmlFile.WriteLine("<a name='" + anchorName + "'></a>");
+				}				
+
+				htmlFile.WriteLine("<h2 class='" + headingClass + "'>" + title + "</h2>");
+				htmlFile.WriteLine("<div class='" + divClass + "'>");
+				htmlFile.WriteLine("<div class='collapsibleSectionInner'>");
+			}
+		}
+
+		protected void EndHtmlSection(System.IO.StreamWriter htmlFile)
+		{
+			if (htmlFile != null)
+			{
+				htmlFile.WriteLine("</div>");
+				htmlFile.WriteLine("</div>");
+			}
+		}
+
+
+		string title;
+		StringBuilder pendingLine = new StringBuilder();
+		List<string> lines = new List<string>();
+		bool bStartCollapsed;
+		string anchorName;
+	};
 
 
 
diff --git a/Engine/Source/Programs/CSVTools/PerfReportTool/SummaryTable.cs b/Engine/Source/Programs/CSVTools/PerfReportTool/SummaryTable.cs
--- a/Engine/Source/Programs/CSVTools/PerfReportTool/SummaryTable.cs
+++ b/Engine/Source/Programs/CSVTools/PerfReportTool/SummaryTable.cs
@@ -2020,7 +2020,7 @@ namespace PerfSummaries
 			}
 
 			htmlFile.WriteLine("<p style='font-size:8'>Created with PerfReportTool " + VersionString + extraString + "</p>");
-			htmlFile.WriteLine("</font></body></html>");
+			htmlFile.WriteLine("</body></html>");
 
 			htmlFile.Close();
 		}
