PerfreportTool 4.227.1 : Fix extraLinks summary null dereference. Add utf-8 charset to meta

diff --git a/Engine/Source/Programs/CSVTools/PerfReportTool/PerfReportTool.cs b/Engine/Source/Programs/CSVTools/PerfReportTool/PerfReportTool.cs
--- a/Engine/Source/Programs/CSVTools/PerfReportTool/PerfReportTool.cs
+++ b/Engine/Source/Programs/CSVTools/PerfReportTool/PerfReportTool.cs
@@ -23,7 +23,7 @@ namespace PerfReportTool
     class Version
     {
 		// Format: Major.Minor.Bugfix
-        private static string VersionString = "4.227.0";
+        private static string VersionString = "4.227.1";
 
         public static string Get() { return VersionString; }
     };
@@ -1312,6 +1312,8 @@ namespace PerfReportTool
 				htmlFile.WriteLine("<html>");
 				htmlFile.WriteLine("  <head>");
 				htmlFile.WriteLine("    <meta http-equiv='X-UA-Compatible' content='IE=edge'/>");
+				htmlFile.WriteLine("    <meta charset='UTF-8'/>");
+
 				if (GetBoolArg("noWatermarks"))
 				{
 					htmlFile.WriteLine("    <![CDATA[ \nCreated with PerfReportTool");
diff --git a/Engine/Source/Programs/CSVTools/PerfReportTool/Summary.cs b/Engine/Source/Programs/CSVTools/PerfReportTool/Summary.cs
--- a/Engine/Source/Programs/CSVTools/PerfReportTool/Summary.cs
+++ b/Engine/Source/Programs/CSVTools/PerfReportTool/Summary.cs
@@ -95,6 +95,10 @@ namespace PerfSummaries
 
         public void ReadStatsFromXML(XElement element, XmlVariableMappings vars)
         {
+			if (element == null) 
+			{
+				return;
+			}
 			useUnstrippedCsvStats = element.GetSafeAttribute<bool>(vars, "useUnstrippedCsvStats", false);
 			bStartCollapsed = element.GetSafeAttribute<bool>(vars, "collapsed", false);
 			XElement statsElement = element.Element("stats");
