PerfReportTool 4.29.0 - Detailed report improvements: Add a top level nav bar with buttons to expand and collapse all sections and jump to the graphs. Clicking on a graph in the list now expands it

diff --git a/Engine/Source/Programs/CSVTools/PerfReportTool/PerfReportTool.cs b/Engine/Source/Programs/CSVTools/PerfReportTool/PerfReportTool.cs
--- a/Engine/Source/Programs/CSVTools/PerfReportTool/PerfReportTool.cs
+++ b/Engine/Source/Programs/CSVTools/PerfReportTool/PerfReportTool.cs
@@ -23,7 +23,7 @@ namespace PerfReportTool
     class Version
     {
 		// Format: Major.Minor.Bugfix
-        private static string VersionString = "4.228.0";
+        private static string VersionString = "4.229.0";
 
         public static string Get() { return VersionString; }
     };
@@ -1368,17 +1368,50 @@ namespace PerfReportTool
 				// Scripting for collapsibles
 				htmlFile.WriteLine("    <script>");
 				htmlFile.WriteLine("        document.addEventListener('DOMContentLoaded', function (event) { setupCollapsibles(); })");
-				htmlFile.WriteLine("        function setupCollapsibles() {");
 
+				htmlFile.WriteLine("        function setupCollapsibles() {");
 				htmlFile.WriteLine("            var collapsibles = document.getElementsByClassName('collapsibleHeading');");
 				htmlFile.WriteLine("            var i;");
 				htmlFile.WriteLine("            for (i = 0; i < collapsibles.length; i++) {");
 				htmlFile.WriteLine("                collapsibles[i].addEventListener('click', function() {");
-				htmlFile.WriteLine("                this.classList.toggle('expanded');");
-				htmlFile.WriteLine("                this.nextElementSibling.classList.toggle('expanded');");
+				htmlFile.WriteLine("                    this.classList.toggle('expanded');");
+				htmlFile.WriteLine("                    this.nextElementSibling.classList.toggle('expanded');");
 				htmlFile.WriteLine("             	});");
 				htmlFile.WriteLine("            }");
 				htmlFile.WriteLine("        }");
+
+				htmlFile.WriteLine("        function collapseAll() {");
+				htmlFile.WriteLine("            var collapsibles = document.getElementsByClassName('collapsibleHeading');");
+				htmlFile.WriteLine("            var i;");
+				htmlFile.WriteLine("            for (i = 0; i < collapsibles.length; i++) {");
+				htmlFile.WriteLine("                setSectionExpanded( collapsibles[i], false );");
+				htmlFile.WriteLine("            }");
+				htmlFile.WriteLine("        }");
+
+				htmlFile.WriteLine("        function expandAll() {");
+				htmlFile.WriteLine("            var collapsibles = document.getElementsByClassName('collapsibleHeading');");
+				htmlFile.WriteLine("            var i;");
+				htmlFile.WriteLine("            for (i = 0; i < collapsibles.length; i++) {");
+				htmlFile.WriteLine("                setSectionExpanded( collapsibles[i], true );");
+				htmlFile.WriteLine("            }");
+				htmlFile.WriteLine("        }");
+
+
+				htmlFile.WriteLine("        function setSectionExpanded(collapsible, bExpanded) {");
+				htmlFile.WriteLine("            if (bExpanded) {");
+				htmlFile.WriteLine("                collapsible.classList.add('expanded');");
+				htmlFile.WriteLine("                collapsible.nextElementSibling.classList.add('expanded');");
+				htmlFile.WriteLine("            }");
+				htmlFile.WriteLine("            else {");
+				htmlFile.WriteLine("                collapsible.classList.remove('expanded');");
+				htmlFile.WriteLine("                collapsible.nextElementSibling.classList.remove('expanded');");
+				htmlFile.WriteLine("            }");
+				htmlFile.WriteLine("        }");
+
+				htmlFile.WriteLine("        function setSectionExpandedById(id, bExpanded) {");
+				htmlFile.WriteLine("            setSectionExpanded(document.getElementById(id), bExpanded);");
+				htmlFile.WriteLine("        }");
+
 				htmlFile.WriteLine("    </script>");
 
 				// CSS
@@ -1390,6 +1423,8 @@ namespace PerfReportTool
 				htmlFile.WriteLine("      h1 {  font-family: 'Verdana', Times, serif; font-size: 20px; padding-top:10px }");
 				htmlFile.WriteLine("      h2 {  font-family: 'Verdana', Times, serif; font-size: 18px; padding-top:5px; padding-bottom:0px; margin-block-end: 0.4em }");
 				htmlFile.WriteLine("      h3 {  font-family: 'Verdana', Times, serif; font-size: 16px; padding-top:5px }");
+				htmlFile.WriteLine("      hr {  margin-top:15px }");
+				htmlFile.WriteLine("      a {  font-family: 'Verdana', Times, serif; font-size: 12px }");
 
 				// Collapsibles
 				htmlFile.WriteLine("      .collapsibleHeading { background-color: #ffffff; cursor: pointer; width: fit-content; border: none; text-align: left; outline: none; display: flex; justify-content: flex-start; align-items: flex-end; flex-direction: row; flex-wrap: nowrap; }");
@@ -1404,7 +1439,7 @@ namespace PerfReportTool
 				htmlFile.WriteLine("    </style>");
 				htmlFile.WriteLine("  </head>");
 				htmlFile.WriteLine("  <body>");
-				htmlFile.WriteLine("  <h1>" + titleStr + "</h1>");
+				htmlFile.WriteLine("  <h1 name='top'>" + titleStr + "</h1>");
 
 				// show the range
 				if (minX > 0 || maxX < Int32.MaxValue)
@@ -1417,6 +1452,7 @@ namespace PerfReportTool
 					htmlFile.WriteLine(")</font>");
 				}
 
+				// Output the metadata table
 				htmlFile.WriteLine("<table style='width:800'>");
 
 				if (reportTypeInfo.metadataToShowList != null)
@@ -1440,6 +1476,13 @@ namespace PerfReportTool
 				htmlFile.WriteLine("<tr><td bgcolor='#F0F0F0'>Frame count</td><td>" + csvStats.SampleCount + " (" + numFramesStripped + " excluded)</td></tr>");
 				htmlFile.WriteLine("</table>");
 
+				// Output the top level nav
+				htmlFile.WriteLine("<br>");
+				htmlFile.WriteLine("<div style='width: 100%; background-color: #ffffff'>");
+				htmlFile.WriteLine("<button type='button' onclick='collapseAll()' style='margin-right:10px'>Collapse all</button>");
+				htmlFile.WriteLine("<button type='button' onclick='expandAll()' style='margin-right:10px'>Expand all</button>");
+				htmlFile.WriteLine("<button type='button' onclick='location.href=\"#graphList\";'>Graphs</button>");
+				htmlFile.WriteLine("</div>");
 			}
 
 			if (summaryRowData != null)
@@ -1474,6 +1517,7 @@ namespace PerfReportTool
 			if (htmlFile != null)
 			{
 				// Output the list of graphs
+				htmlFile.WriteLine("<hr><a name='graphList'></a>");
 				htmlFile.WriteLine("<h2>Graphs</h2>");
 
 				// TODO: support sections for graphs
@@ -1498,11 +1542,12 @@ namespace PerfReportTool
 					{
 						string svgTitle = csvSvgInfo.Graph.title;
 						// TODO: Check if this graph belongs in this section.
-						htmlFile.WriteLine("<li><a href='#" + StripSpaces(svgTitle) + "'>" + svgTitle + "</a></li>");
+						htmlFile.WriteLine("<li><a href='#" + StripSpaces(svgTitle) + "' onclick='setSectionExpandedById(\""+ StripSpaces(svgTitle) + "\", true)'>" + svgTitle + "</a></li>");
 					}
 
 					htmlFile.WriteLine("</ul>");
 				}
+				htmlFile.WriteLine("<a href='#top'>Back to top \u2191</a>");
 
 
 				// Output the Graphs
@@ -1570,6 +1615,9 @@ namespace PerfReportTool
 					htmlSection.WriteToFile(htmlFile);
 				}
 
+				htmlFile.WriteLine("<a href='#top'>Back to top \u2191</a>");
+
+
 				if (GetBoolArg("noWatermarks"))
 				{
 					htmlFile.WriteLine("<p style='font-size:8'>Created with PerfReportTool</p>");
diff --git a/Engine/Source/Programs/CSVTools/PerfReportTool/Summary.cs b/Engine/Source/Programs/CSVTools/PerfReportTool/Summary.cs
--- a/Engine/Source/Programs/CSVTools/PerfReportTool/Summary.cs
+++ b/Engine/Source/Programs/CSVTools/PerfReportTool/Summary.cs
@@ -248,11 +248,11 @@ namespace PerfSummaries
 
 	class HtmlSection
 	{
-		public HtmlSection(string titleIn, bool bStartCollapsedIn, string anchorNameIn = null)
+		public HtmlSection(string titleIn, bool bStartCollapsedIn, string elementIdIn = null)
 		{
 			title = titleIn;
 			bStartCollapsed = bStartCollapsedIn;
-			anchorName = anchorNameIn; 
+			elementId = elementIdIn; 
 		}
 
 		public void WriteLine(string text) 
@@ -298,12 +298,13 @@ namespace PerfSummaries
 					divClass += " expanded";
 				}
 
-				if (anchorName != null)
+				string extraHeadingAttributes = "";
+				if (elementId != null)
 				{
-					htmlFile.WriteLine("<a name='" + anchorName + "'></a>");
-				}				
-
-				htmlFile.WriteLine("<h2 class='" + headingClass + "'>" + title + "</h2>");
+					htmlFile.WriteLine("<a name='" + elementId + "'></a>");
+					extraHeadingAttributes = "id='" + elementId + "'";
+				}
+				htmlFile.WriteLine("<h2 class='" + headingClass + "' "+ extraHeadingAttributes + ">" + title + "</h2>");
 				htmlFile.WriteLine("<div class='" + divClass + "'>");
 				htmlFile.WriteLine("<div class='collapsibleSectionInner'>");
 			}
@@ -323,7 +324,7 @@ namespace PerfSummaries
 		StringBuilder pendingLine = new StringBuilder();
 		List<string> lines = new List<string>();
 		bool bStartCollapsed;
-		string anchorName;
+		string elementId;
 	};
 
 
