PerfReportTool: Minor tweaks to heading formatting for better use of vertical space. Graphs now use h3 instead of h2

diff --git a/Engine/Source/Programs/CSVTools/PerfReportTool/PerfReportTool.cs b/Engine/Source/Programs/CSVTools/PerfReportTool/PerfReportTool.cs
--- a/Engine/Source/Programs/CSVTools/PerfReportTool/PerfReportTool.cs
+++ b/Engine/Source/Programs/CSVTools/PerfReportTool/PerfReportTool.cs
@@ -1421,8 +1421,8 @@ namespace PerfReportTool
 				htmlFile.WriteLine("      p {  font-family: 'Verdana', Times, serif; font-size: 12px }");
 				htmlFile.WriteLine("      ul {  font-family: 'Verdana', Times, serif; font-size: 14px }");
 				htmlFile.WriteLine("      h1 {  font-family: 'Verdana', Times, serif; font-size: 20px; padding-top:10px }");
-				htmlFile.WriteLine("      h2 {  font-family: 'Verdana', Times, serif; font-size: 18px; padding-top:5px; padding-bottom:0px; margin-block-end: 0.4em }");
-				htmlFile.WriteLine("      h3 {  font-family: 'Verdana', Times, serif; font-size: 16px; padding-top:5px }");
+				htmlFile.WriteLine("      h2 {  font-family: 'Verdana', Times, serif; font-size: 18px; padding-top:5px; padding-bottom:0px; margin-block-start: 0.4em; margin-block-end: 0.4em }");
+				htmlFile.WriteLine("      h3 {  font-family: 'Verdana', Times, serif; font-size: 16px; padding-top:5px; margin-block-start: 0.3em; margin-block-end: 0.3em }");
 				htmlFile.WriteLine("      hr {  margin-top:15px }");
 				htmlFile.WriteLine("      a {  font-family: 'Verdana', Times, serif; font-size: 12px }");
 
@@ -1557,7 +1557,7 @@ namespace PerfReportTool
 					ReportGraph graph = csvSvgInfo.Graph;
 
 					string svgTitle = graph.title;
-					HtmlSection htmlSection = new HtmlSection(svgTitle, false, StripSpaces(svgTitle));
+					HtmlSection htmlSection = new HtmlSection(svgTitle, false, StripSpaces(svgTitle), 3);
 
 					if (csvSvgInfo.Format == CsvSvgInfo.GraphFormat.Inline)
 					{
diff --git a/Engine/Source/Programs/CSVTools/PerfReportTool/Summary.cs b/Engine/Source/Programs/CSVTools/PerfReportTool/Summary.cs
--- a/Engine/Source/Programs/CSVTools/PerfReportTool/Summary.cs
+++ b/Engine/Source/Programs/CSVTools/PerfReportTool/Summary.cs
@@ -248,11 +248,12 @@ namespace PerfSummaries
 
 	class HtmlSection
 	{
-		public HtmlSection(string titleIn, bool bStartCollapsedIn, string elementIdIn = null)
+		public HtmlSection(string titleIn, bool bStartCollapsedIn, string elementIdIn = null, int headingLevel=2)
 		{
 			title = titleIn;
 			bStartCollapsed = bStartCollapsedIn;
-			elementId = elementIdIn; 
+			elementId = elementIdIn;
+			headingType = "h" + headingLevel.ToString();
 		}
 
 		public void WriteLine(string text) 
@@ -304,7 +305,7 @@ namespace PerfSummaries
 					htmlFile.WriteLine("<a name='" + elementId + "'></a>");
 					extraHeadingAttributes = "id='" + elementId + "'";
 				}
-				htmlFile.WriteLine("<h2 class='" + headingClass + "' "+ extraHeadingAttributes + ">" + title + "</h2>");
+				htmlFile.WriteLine("<"+ headingType + " class='" + headingClass + "' "+ extraHeadingAttributes + ">" + title + "</"+ headingType + ">");
 				htmlFile.WriteLine("<div class='" + divClass + "'>");
 				htmlFile.WriteLine("<div class='collapsibleSectionInner'>");
 			}
@@ -314,6 +315,7 @@ namespace PerfSummaries
 		{
 			if (htmlFile != null)
 			{
+				htmlFile.WriteLine("<br>");
 				htmlFile.WriteLine("</div>");
 				htmlFile.WriteLine("</div>");
 			}
@@ -325,6 +327,7 @@ namespace PerfSummaries
 		List<string> lines = new List<string>();
 		bool bStartCollapsed;
 		string elementId;
+		string headingType;
 	};
 
 
