PRT - Update help text to include new dynamic graph args.

Also removed deprecated externalGraphs arg

#rnx

diff --git a/Engine/Source/Programs/CSVTools/PerfReportTool/PerfReportTool.cs b/Engine/Source/Programs/CSVTools/PerfReportTool/PerfReportTool.cs
--- a/Engine/Source/Programs/CSVTools/PerfReportTool/PerfReportTool.cs
+++ b/Engine/Source/Programs/CSVTools/PerfReportTool/PerfReportTool.cs
@@ -23,7 +23,7 @@ namespace PerfReportTool
     class Version
     {
 		// Format: Major.Minor.Bugfix
-        private static string VersionString = "4.229.0";
+        private static string VersionString = "4.229.1";
 
         public static string Get() { return VersionString; }
     };
@@ -107,6 +107,8 @@ namespace PerfReportTool
 			"  -graphThreads : use with -batchedGraphs to control the number of threads per CsvToSVG instance \n" +
 			"                  (default: PC core count/2)\n" +
 			"  -csvToSvgSequential : Run CsvToSvg sequentially\n" +
+			"  -useEmbeddedGraphUrl : Insert a script to fetch the graph from the specified endpoint on page load rather than embedding the full graph\n" +
+			"  -embeddedGraphUrlRoot : The url to fetch the graph from if -useEmbeddedGraphUrl is specified. CsvToSvg graph args are provided as get params\n" +
 			"Deprecated performance args:\n" +
 			"  -csvToSvgProcesses : Use separate processes for csvToSVG instead of threads (slower)\n" +
             "  -embedGraphCommandline : if -csvToSvgProcesses is specified, embeds the commandline for debugging purposes\n" +
@@ -142,7 +144,6 @@ namespace PerfReportTool
 			"     e.g \"platform=ps4 AND deviceprofile=ps4_60\" \n" +
 			"  -readAllStats : allows any CSV stat avg to appear in the summary table, not just those referenced in summaries\n" +
 			"  -showHiddenStats : shows stats which have been automatically hidden (typically duplicate csv unit stats)\n" +
-			"  -externalGraphs : enables external graphs (off by default)\n" +
 			"  -spreadsheetfriendly: outputs a single quote before non-numeric entries in summary tables\n" +
 			"  -noSummaryMinMax: don't make min/max columns for each stat in a condensed summary\n" +
 			"  -reverseTable [0|1]: Reverses the order of summary tables (set 0 to force off)\n" +
