Adding -showmin and -showmax to csvinfo

#rnx

diff --git a/Engine/Source/Programs/CSVTools/CsvInfo/CSVInfo.cs b/Engine/Source/Programs/CSVTools/CsvInfo/CSVInfo.cs
--- a/Engine/Source/Programs/CSVTools/CsvInfo/CSVInfo.cs
+++ b/Engine/Source/Programs/CSVTools/CsvInfo/CSVInfo.cs
@@ -73,6 +73,8 @@ namespace CSVInfo
                 "Format: \n" +
                 "  <csvfilename>\n"+
 				"  [-showaverages] Output stat averages.\n" +
+				"  [-showmin] Output stat min values.\n" +
+				"  [-showmax] Output stat max values.\n" +
 				"  [-showTotals] Output stat totals.\n" +
 				"  [-showAllStats] Output stat values from all frames.\n" +
 				"  [-showEvents] Output CSV Events and the frame they occured on.\n" +
@@ -94,12 +96,15 @@ namespace CSVInfo
             ReadCommandLine(args);
 
             bool showAverages = GetBoolArg("showAverages");
+            bool showMin = GetBoolArg("showMin");
+            bool showMax = GetBoolArg("showMax");
 			bool showTotals= GetBoolArg("showTotals");
 			bool showAllStats = GetBoolArg("showAllStats");
 			bool showEvents = GetBoolArg("showEvents");
 			string jsonFilename = GetArg("toJson",false);
 			string statFilterString = GetArg("statFilters",false);
-			bool bReadJustHeader = !GetBoolArg("forceFullRead") && !showAverages && !showTotals && (statFilterString.Length == 0);
+			bool bReadJustHeader = !GetBoolArg("forceFullRead") && !showAverages && !showTotals && !showMin && !showMax
+			                       && (statFilterString.Length == 0);
 
 			string[] statFilters = statFilterString.Length > 0 ? statFilterString.Split(',') : null;
 
@@ -158,7 +163,7 @@ namespace CSVInfo
 				}
 				statLines.Sort();
 
-				if (showTotals || showAverages || showAllStats)
+				if (showTotals || showAverages || showMin || showMax || showAllStats)
 				{
 					jsonLines.Add("  \"stats\": {");
 					for (int i=0; i<statLines.Count; i++)
@@ -173,6 +178,14 @@ namespace CSVInfo
 						{
 							entries.Add("\"average\":" + csvStats.GetStat(statName).average);
 						}
+						if (showMin)
+						{
+							entries.Add("\"min\":" + csvStats.GetStat(statName).ComputeMinValue());
+						}
+						if (showMax)
+						{
+							entries.Add("\"max\":" + csvStats.GetStat(statName).ComputeMaxValue());
+						}
 						if (showAllStats)
 						{
 							entries.Add("\"all\": [" + String.Join(",", csvStats.GetStat(statName).samples) + "]");
@@ -211,6 +224,14 @@ namespace CSVInfo
 					{
 						statLine += " (Total: " + stat.total.ToString()+") ";
 					}
+					if (showMin)
+					{
+						statLine += " (Min: " + stat.ComputeMinValue().ToString()+") ";
+					}
+					if (showMax)
+					{
+						statLine += " (Max: " + stat.ComputeMaxValue().ToString()+") ";
+					}
 					if (showAllStats)
 					{
 						statLine += " (All: [" + String.Join(",", stat.samples) + "]) ";
