PerfReportTool 4.230.0 - Add support for variables which can be used outside of individual CSV reporting. Any variable declared in a variableSet without a metadata query can be used in this way. This is useful if we want to access variables in SummaryTable definitions for example. Also strip whitespace in summary table row sort and filter lists

diff --git a/Engine/Source/Programs/CSVTools/PerfReportTool/PerfReportTool.cs b/Engine/Source/Programs/CSVTools/PerfReportTool/PerfReportTool.cs
--- a/Engine/Source/Programs/CSVTools/PerfReportTool/PerfReportTool.cs
+++ b/Engine/Source/Programs/CSVTools/PerfReportTool/PerfReportTool.cs
@@ -23,7 +23,7 @@ namespace PerfReportTool
     class Version
     {
 		// Format: Major.Minor.Bugfix
-        private static string VersionString = "4.229.1";
+        private static string VersionString = "4.230.0";
 
         public static string Get() { return VersionString; }
     };
diff --git a/Engine/Source/Programs/CSVTools/PerfReportTool/ReportXml.cs b/Engine/Source/Programs/CSVTools/PerfReportTool/ReportXml.cs
--- a/Engine/Source/Programs/CSVTools/PerfReportTool/ReportXml.cs
+++ b/Engine/Source/Programs/CSVTools/PerfReportTool/ReportXml.cs
@@ -186,7 +186,15 @@ namespace PerfReportTool
 				throw new Exception("No reporttypes element found in report XML " + reportTypeXmlFilename);
 			}
 
+			// Read the global element set
 			globalVariableSetElement = rootElement.Element("globalVariableSet");
+			if (globalVariableSetElement != null)
+			{
+				// Read static variable mappings. That is all the variables which don't depend on CSV metadata (anything with a metadata query will be stripped)
+                // These variables are independent of CSVs and can be used outside of individual reports (e.g in summary tables)
+				staticVariableMappings = new XmlVariableMappings();
+				staticVariableMappings.ApplyVariableSet(globalVariableSetElement, null);
+			}
 
 			// Read the graph XML
 			string graphsXMLFilename;
@@ -207,7 +215,7 @@ namespace PerfReportTool
 				}
 
 			}
-			defaultReportTypeName = reportTypesElement.GetSafeAttribute<string>("default");
+			defaultReportTypeName = reportTypesElement.GetSafeAttribute<string>(staticVariableMappings, "default");
 
 			Console.Out.WriteLine("GraphXML:  " + graphsXMLFilename+"\n");
 			XDocument reportGraphsDoc = XDocument.Load(graphsXMLFilename);
@@ -325,7 +333,7 @@ namespace PerfReportTool
 				summaryTables = new Dictionary<string, SummaryTableInfo>();
 				foreach (XElement summaryElement in summaryTablesElement.Elements("summaryTable"))
 				{
-					SummaryTableInfo table = new SummaryTableInfo(summaryElement, substitutionsDict, appendList, rowSortAppendList);
+					SummaryTableInfo table = new SummaryTableInfo(summaryElement, substitutionsDict, appendList, rowSortAppendList, staticVariableMappings);
 					summaryTables.Add(summaryElement.GetRequiredAttribute<string>("name").ToLower(), table);
 				}
 			}
@@ -562,6 +570,7 @@ namespace PerfReportTool
 		XElement graphGroupsElement;
 		XElement summaryTablesElement;
 		XElement globalVariableSetElement;
+		XmlVariableMappings staticVariableMappings;
 		string defaultReportTypeName;
 		Dictionary<string, XElement> sharedSummaries;
 		Dictionary<string, GraphSettings> graphs;
diff --git a/Engine/Source/Programs/CSVTools/PerfReportTool/SummaryTable.cs b/Engine/Source/Programs/CSVTools/PerfReportTool/SummaryTable.cs
--- a/Engine/Source/Programs/CSVTools/PerfReportTool/SummaryTable.cs
+++ b/Engine/Source/Programs/CSVTools/PerfReportTool/SummaryTable.cs
@@ -64,12 +64,12 @@ namespace PerfSummaries
 
 	class SummaryTableInfo
 	{
-		public SummaryTableInfo(XElement tableElement, Dictionary<string,string> substitutionsDict, string[] appendList, string[] rowSortAppendList)
+		public SummaryTableInfo(XElement tableElement, Dictionary<string,string> substitutionsDict, string[] appendList, string[] rowSortAppendList, XmlVariableMappings variableMappings )
 		{
-			XAttribute rowSortAt = tableElement.Attribute("rowSort");
-			if (rowSortAt != null)
+			string rowSortStr = tableElement.GetSafeAttribute<string>(variableMappings, "rowSort");
+			if (rowSortStr != null)
 			{
-				rowSortList.AddRange(rowSortAt.Value.Split(','));
+				rowSortList.AddRange(rowSortStr.Split(',').Select(s => s.Trim()));
 				ApplySubstitutionsToList(rowSortList, substitutionsDict);
 			}
 			if (rowSortAppendList != null)
@@ -77,16 +77,16 @@ namespace PerfSummaries
 				rowSortList.AddRange(rowSortAppendList);
 			}
 
-			XAttribute weightByColumnAt = tableElement.Attribute("weightByColumn");
-			if (weightByColumnAt != null)
+			weightByColumn = tableElement.GetSafeAttribute<string>(variableMappings, "weightByColumn");
+			if (weightByColumn != null)
 			{
-				weightByColumn = weightByColumnAt.Value.ToLower();
+				weightByColumn = weightByColumn.ToLower();
 			}
 
 			XElement filterEl = tableElement.Element("filter");
 			if (filterEl != null)
 			{
-				columnFilterList.AddRange(filterEl.Value.Split(','));
+				columnFilterList.AddRange(filterEl.GetValue(variableMappings).Split(',').Select(s => s.Trim()));
 				ApplySubstitutionsToList(columnFilterList, substitutionsDict);
 			}
 
@@ -95,10 +95,10 @@ namespace PerfSummaries
 				columnFilterList.AddRange(appendList);
 			}
 
-			bReverseSortRows = tableElement.GetSafeAttribute<bool>("reverseSortRows", false);
-			bScrollableFormatting = tableElement.GetSafeAttribute<bool>("scrollableFormatting", false);
+			bReverseSortRows = tableElement.GetSafeAttribute<bool>(variableMappings, "reverseSortRows", false);
+			bScrollableFormatting = tableElement.GetSafeAttribute<bool>(variableMappings, "scrollableFormatting", false);
 
-			string colorizeModeStr = tableElement.GetSafeAttribute<string>("colorizeMode", "").ToLower();
+			string colorizeModeStr = tableElement.GetSafeAttribute<string>(variableMappings, "colorizeMode", "").ToLower();
 			if (colorizeModeStr != "")
 			{
 				if (colorizeModeStr == "auto")
@@ -115,22 +115,22 @@ namespace PerfSummaries
 				}
 			}
 
-			statThreshold = tableElement.GetSafeAttribute<float>("statThreshold", 0.0f);
-			hideStatPrefix = tableElement.GetSafeAttribute<string>("hideStatPrefix");
+			statThreshold = tableElement.GetSafeAttribute<float>(variableMappings, "statThreshold", 0.0f);
+			hideStatPrefix = tableElement.GetSafeAttribute<string>(variableMappings, "hideStatPrefix");
 
 			foreach (XElement sectionBoundaryEl in tableElement.Elements("sectionBoundary"))
 			{
 				if (sectionBoundaryEl != null)
 				{
-					string statName = ApplySubstitution(sectionBoundaryEl.GetSafeAttribute<string>("statName"), substitutionsDict);
+					string statName = ApplySubstitution(sectionBoundaryEl.GetSafeAttribute<string>(variableMappings, "statName"), substitutionsDict);
 
 					SummarySectionBoundaryInfo sectionBoundary = new SummarySectionBoundaryInfo(
 						statName,
-						sectionBoundaryEl.GetSafeAttribute<string>("startToken"),
-						sectionBoundaryEl.GetSafeAttribute<string>("endToken"),
-						sectionBoundaryEl.GetSafeAttribute<int>("level", 0),
-						sectionBoundaryEl.GetSafeAttribute<bool>("inCollatedTable", true),
-						sectionBoundaryEl.GetSafeAttribute<bool>("inFullTable", true)
+						sectionBoundaryEl.GetSafeAttribute<string>(variableMappings, "startToken"),
+						sectionBoundaryEl.GetSafeAttribute<string>(variableMappings, "endToken"),
+						sectionBoundaryEl.GetSafeAttribute<int>(variableMappings, "level", 0),
+						sectionBoundaryEl.GetSafeAttribute<bool>(variableMappings, "inCollatedTable", true),
+						sectionBoundaryEl.GetSafeAttribute<bool>(variableMappings, "inFullTable", true)
 						);
 					sectionBoundaries.Add(sectionBoundary);
 				}
@@ -165,8 +165,8 @@ namespace PerfSummaries
 
 		public SummaryTableInfo(string filterListStr, string rowSortStr)
 		{
-			columnFilterList.AddRange(filterListStr.Split(','));
-			rowSortList.AddRange(rowSortStr.Split(','));
+			columnFilterList.AddRange(filterListStr.Split(',').Select(s => s.Trim()));
+			rowSortList.AddRange(rowSortStr.Split(',').Select(s => s.Trim()));
 		}
 
 		public SummaryTableInfo()
