PerfReportTool: When stat names are remapped, display the original name as a tooltip

diff --git a/Engine/Source/Programs/CSVTools/PerfReportTool/PerfReportTool.cs b/Engine/Source/Programs/CSVTools/PerfReportTool/PerfReportTool.cs
--- a/Engine/Source/Programs/CSVTools/PerfReportTool/PerfReportTool.cs
+++ b/Engine/Source/Programs/CSVTools/PerfReportTool/PerfReportTool.cs
@@ -23,7 +23,7 @@ namespace PerfReportTool
     class Version
     {
 		// Format: Major.Minor.Bugfix
-        private static string VersionString = "4.230.0";
+        private static string VersionString = "4.231.0";
 
         public static string Get() { return VersionString; }
     };
diff --git a/Engine/Source/Programs/CSVTools/PerfReportTool/SummaryTable.cs b/Engine/Source/Programs/CSVTools/PerfReportTool/SummaryTable.cs
--- a/Engine/Source/Programs/CSVTools/PerfReportTool/SummaryTable.cs
+++ b/Engine/Source/Programs/CSVTools/PerfReportTool/SummaryTable.cs
@@ -2,6 +2,7 @@
 
 using System;
 using System.Collections.Generic;
+using System.Data.Common;
 using System.Linq;
 using System.Xml.Linq;
 using PerfReportTool;
@@ -835,6 +836,12 @@ namespace PerfSummaries
 			{
 				attributes.Add($"title='{tooltip}'");
 			}
+			else if ( displayName != null )
+			{
+				string baseStatName = SummaryTable.GetBaseStatNameWithPrefixAndSuffix(name, out _, out _);
+
+				attributes.Add($"title='{baseStatName}'");
+			}
 			return attributes;
 		}
 
@@ -1756,7 +1763,7 @@ namespace PerfSummaries
 				// Add the special columns (up to Count) to the lower header row
 				for (int i = 0; i < firstStatColumnIndex; i++)
 				{
-					headerRow.AddCell(columns[i].GetDisplayName(hideStatPrefix, bAddStatNameSpacing, bGreyOutStatPrefixes));
+					headerRow.AddCell(columns[i].GetDisplayName(hideStatPrefix, bAddStatNameSpacing, bGreyOutStatPrefixes), String.Join(" ", columns[i].GetHeaderAttributes()));
 				}
 
 				if (bAddMinMaxColumns)
