PerfReportTool: Also output the elementType in the header tooltip (E.g CsvMetadata, SummaryTableMetric etc)

diff --git a/Engine/Source/Programs/CSVTools/PerfReportTool/SummaryTable.cs b/Engine/Source/Programs/CSVTools/PerfReportTool/SummaryTable.cs
--- a/Engine/Source/Programs/CSVTools/PerfReportTool/SummaryTable.cs
+++ b/Engine/Source/Programs/CSVTools/PerfReportTool/SummaryTable.cs
@@ -836,11 +836,11 @@ namespace PerfSummaries
 			{
 				attributes.Add($"title='{tooltip}'");
 			}
-			else if ( displayName != null )
+			else 
 			{
-				string baseStatName = SummaryTable.GetBaseStatNameWithPrefixAndSuffix(name, out _, out _);
-
-				attributes.Add($"title='{baseStatName}'");
+				string tooltipStr = SummaryTable.GetBaseStatNameWithPrefixAndSuffix(name, out _, out _);
+				tooltipStr += " (" + this.elementType.ToString() + ")";
+				attributes.Add($"title='{tooltipStr}'");
 			}
 			return attributes;
 		}
