PerfReportTool 4.232.0: Add support for stat type qualifiers in summary tables so we can identify stats by their type to avoid ambiguity, e.g [meta]deviceprofile. Prefixes are :
[metric] SummaryTableMetric
[csv] CsvStatAverage
[meta] CsvMetadata
[toolmeta] ToolMetadata

diff --git a/Engine/Source/Programs/CSVTools/PerfReportTool/PerfReportTool.cs b/Engine/Source/Programs/CSVTools/PerfReportTool/PerfReportTool.cs
--- a/Engine/Source/Programs/CSVTools/PerfReportTool/PerfReportTool.cs
+++ b/Engine/Source/Programs/CSVTools/PerfReportTool/PerfReportTool.cs
@@ -23,7 +23,7 @@ namespace PerfReportTool
     class Version
     {
 		// Format: Major.Minor.Bugfix
-        private static string VersionString = "4.231.0";
+        private static string VersionString = "4.232.0";
 
         public static string Get() { return VersionString; }
     };
@@ -177,7 +177,6 @@ namespace PerfReportTool
 			"  -columnDiffDisplayThreshold <value> : if specified, hides columns with max diff value below this threshold\n"+
 			"\n" +
 			"Optional Column Filters\n" +
-			"  -debugShowFilteredColumns : grays out filtered columns instead of removing. Column tooltip will show filtered reason.\n" +
 			"  -hideMetadataColumns : filters out metadata columns from the table (excluding those used in row sort).\n" +
 			"\n" +
 			"Regression Column Filtering\n" +
@@ -795,13 +794,12 @@ namespace PerfReportTool
 			}
 
 			IEnumerable<ISummaryTableColumnFilter> additionalColumnFilters = MakeAdditionalColumnFilters(tableInfo);
-			bool showFilteredColumns = GetBoolArg("debugShowFilteredColumns");
 
 			// Set format info for the columns as some of the info is needed for the filters.
 			// TODO: would be better if we could determine HighIsBad without the format info and store it directly in the column.
 			table.SetColumnFormatInfo(reportXML.columnFormatInfoList);
 
-			SummaryTable filteredTable = table.SortAndFilter(tableInfo.columnFilterList, tableInfo.rowSortList, bReverseTable, weightByColumnName, showFilteredColumns, additionalColumnFilters);
+			SummaryTable filteredTable = table.SortAndFilter(tableInfo.columnFilterList, tableInfo.rowSortList, bReverseTable, weightByColumnName, additionalColumnFilters);
 			if (bCollated)
 			{
 				filteredTable = filteredTable.CollateSortedTable(tableInfo.rowSortList, addMinMaxColumns);
@@ -835,8 +833,7 @@ namespace PerfReportTool
 					GetIntArg("maxSummaryTableStringLength", Int32.MaxValue), 
 					weightByColumnName, 
 					summaryTitle,
-					bCollated ? bTransposeCollatedSummaryTable : bTransposeFullSummaryTable,
-					showFilteredColumns
+					bCollated ? bTransposeCollatedSummaryTable : bTransposeFullSummaryTable
 				);
 			}
 		}
diff --git a/Engine/Source/Programs/CSVTools/PerfReportTool/SummaryTable.cs b/Engine/Source/Programs/CSVTools/PerfReportTool/SummaryTable.cs
--- a/Engine/Source/Programs/CSVTools/PerfReportTool/SummaryTable.cs
+++ b/Engine/Source/Programs/CSVTools/PerfReportTool/SummaryTable.cs
@@ -5,6 +5,7 @@ using System.Collections.Generic;
 using System.Data.Common;
 using System.Linq;
 using System.Xml.Linq;
+using CSVStats;
 using PerfReportTool;
 
 namespace PerfSummaries
@@ -416,6 +417,11 @@ namespace PerfSummaries
 			}
 		}
 
+		public string getKey(bool bIncludeTypeQualifier = true)
+		{
+			return bIncludeTypeQualifier ? SummaryTable.GetElementTypeStatPrefix(this.elementType) + name.ToLower() : name.ToLower();
+		}
+
 		public SummaryTableColumn Clone()
 		{
 			SummaryTableColumn newColumn = new SummaryTableColumn(name, isNumeric, displayName, isRowWeightColumn, elementType, formatInfo, tooltip, columnColourModifier);
@@ -1016,6 +1022,130 @@ namespace PerfSummaries
 
 	class SummaryTable
 	{
+		class SummaryTableColumnLookup
+		{
+			public SummaryTableColumnLookup()
+			{
+				Clear();
+			}
+
+			public void Clear()
+			{
+				for (int i = 0; i < (int)SummaryTableElement.Type.COUNT + 1; i++)
+				{
+					columnLookupByType[i] = new Dictionary<string, SummaryTableColumn>();
+				}
+			}
+
+			public List<string> GetSortedKeyList(SummaryTableElement.Type type = SummaryTableElement.Type.ANY)
+			{
+				List<string> listOut = columnLookupByType[(int)type].Keys.ToList();
+				listOut.Sort();
+				return listOut;
+			}
+
+			public void Add(string key, SummaryTableColumn column)
+			{
+				// Is there a type qualifier prefix in the key? If so, determine the type from the qualifier and stip the prefix
+				if (key.StartsWith("["))
+				{
+					throw new Exception("Key can't start with a name qualifier prefix ([)");
+				}
+
+				// Add to the specific type lookup
+				columnLookupByType[(int)column.elementType][key] = column;
+
+				// Try to add to the global lookup, but handle collisions with priority
+				Dictionary<string, SummaryTableColumn> columnLookupAny = columnLookupByType[(int)SummaryTableElement.Type.ANY];
+				if (columnLookupAny.TryGetValue(key, out SummaryTableColumn existingColumn))
+				{
+					// Handle collisions. These are allowed. If we want to be immune from collisions in lookups, we can specify a type
+					// Prioritize summary table metrics over CSV stats, since this matches existing behaviour
+					if (GetElementTypePriority(column.elementType) > GetElementTypePriority(existingColumn.elementType))
+					{
+						columnLookupAny[key] = column;
+					}
+				}
+				else
+				{
+					columnLookupAny[key] = column;
+				}
+			}
+
+			public bool ContainsKey(string key, SummaryTableElement.Type type = SummaryTableElement.Type.ANY)
+			{
+				return Get(key, type) != null;
+			}
+
+			public void Remove(string key, SummaryTableElement.Type type = SummaryTableElement.Type.ANY)
+			{
+				// Is there a type qualifier prefix in the key? If so, determine the type from the qualifier and stip the prefix
+				if (type == SummaryTableElement.Type.ANY && key.StartsWith("["))
+				{
+					type = SummaryTable.GetQualfiedStatType(key, out key);
+				}
+
+				if (type == SummaryTableElement.Type.ANY)
+				{
+					// If a type isn't specified, we have to remove from all lookup tables, since there may be entries with different types and the same key
+					for (int i = 0; i < (int)SummaryTableElement.Type.COUNT + 1; i++)
+					{
+						columnLookupByType[i].Remove(key);
+					}
+				}
+				else
+				{
+					columnLookupByType[(int)type].Remove(key);
+
+					// Only remove from the global column if the type is the same as specified
+					Dictionary<string, SummaryTableColumn> columnLookupAny = columnLookupByType[(int)SummaryTableElement.Type.ANY];
+					if (columnLookupAny.TryGetValue(key, out SummaryTableColumn column))
+					{
+						if (column.elementType == type)
+						{
+							columnLookupAny.Remove(key);
+						}
+					}
+				}
+
+			}
+
+			public SummaryTableColumn Get(string key, SummaryTableElement.Type type = SummaryTableElement.Type.ANY)
+			{
+				// Is there a type qualifier prefix in the key? If so, determine the type from the qualifier and stip the prefix
+				if (type == SummaryTableElement.Type.ANY && key.StartsWith("["))
+				{
+					type = SummaryTable.GetQualfiedStatType(key, out key);
+				}
+
+				if (columnLookupByType[(int)type].TryGetValue(key, out SummaryTableColumn columnOut))
+				{
+					return columnOut;
+				}
+				return null;
+			}
+
+			// Determines the priority in the event of collisions in the ANY lookup (used when no type is specified)
+			public static int GetElementTypePriority(SummaryTableElement.Type type)
+			{
+				switch (type)
+				{
+					case SummaryTableElement.Type.CsvMetadata:
+						return 2;
+					case SummaryTableElement.Type.CsvStatAverage:
+						return 1;
+					case SummaryTableElement.Type.SummaryTableMetric:
+						return 3;
+					case SummaryTableElement.Type.ToolMetadata:
+						return 4;
+					default:
+						return 5;
+				}
+			}
+
+			Dictionary<string, SummaryTableColumn>[] columnLookupByType = new Dictionary<string, SummaryTableColumn>[(int)SummaryTableElement.Type.COUNT + 1];
+		};
+
 		public SummaryTable()
 		{
 		}
@@ -1039,7 +1169,7 @@ namespace PerfSummaries
 				string key = collateBy.ToLower();
 				if (columnLookup.ContainsKey(key))
 				{
-					collateByColumns.Add(columnLookup[key]);
+					collateByColumns.Add(columnLookup.Get(key));
 				}
 			}
 			if (collateByColumns.Count == 0)
@@ -1100,7 +1230,7 @@ namespace PerfSummaries
 			string CurrentRowSortKey = "";
 			foreach (string collateBy in finalSortByList)
 			{
-				CurrentRowSortKey += "{" + columnLookup[collateBy].GetStringValue(0) + "}";
+				CurrentRowSortKey += "{" + columnLookup.Get(collateBy).GetStringValue(0) + "}";
 			}
 
 			int destRowIndex = 0;
@@ -1162,7 +1292,7 @@ namespace PerfSummaries
 				{
 					foreach (string collateBy in finalSortByList)
 					{
-						nextSortKey += "{" + columnLookup[collateBy].GetStringValue(i + 1) + "}";
+						nextSortKey += "{" + columnLookup.Get(collateBy).GetStringValue(i + 1) + "}";
 					}
 				}
 
@@ -1172,7 +1302,7 @@ namespace PerfSummaries
 					for (int j = 0; j < countColumnIndex; j++)
 					{
 						string key = newColumns[j].name.ToLower();
-						newColumns[j].SetStringValue(destRowIndex, columnLookup[key].GetStringValue(i));
+						newColumns[j].SetStringValue(destRowIndex, columnLookup.Get(key).GetStringValue(i));
 					}
 					// Commit the row 
 					newColumns[countColumnIndex].SetValue(destRowIndex, (double)mergedRowsCount);
@@ -1229,84 +1359,169 @@ namespace PerfSummaries
 
 			string prefix = GetBaseStatPrefix(inColumn.name);
 			string lookupKey = ( SummaryTableColumn.getAggregateTypePrefix(aggregateType) + inColumn.name.Substring(prefix.Length) ).ToLower();
-			if ( !columnLookup.TryGetValue(lookupKey, out SummaryTableColumn outColumn) )
+			SummaryTableColumn outColumn = columnLookup.Get(lookupKey);
+			if (outColumn == null)
 			{
 				throw new Exception("Aggregate column "+lookupKey+" not found!");
 			}
 			return outColumn;
 		}
 
+		SummaryTableColumnLookup GetColumnLookup()
+		{
+			return columnLookup;
+		}
+
 
 		public SummaryTable SortAndFilter(string customFilter, string customRowSort = "buildversion,deviceprofile", bool bReverseSort = false, string weightByColumnName = null)
 		{
 			return SortAndFilter(customFilter.Split(',').ToList(), customRowSort.Split(',').ToList(), bReverseSort, weightByColumnName);
 		}
 
-		public SummaryTable SortAndFilter(List<string> columnFilterList, List<string> rowSortList, bool bReverseSort, string weightByColumnName, bool showFilteredColumns = false, IEnumerable<ISummaryTableColumnFilter> additionalFilters = null)
+
+		public static string GetElementTypeStatPrefix(SummaryTableElement.Type type)
 		{
-			SummaryTable newTable = SortRows(rowSortList, bReverseSort);
+			switch (type)
+			{
+				case SummaryTableElement.Type.SummaryTableMetric:
+					return "[metric]";
+				case SummaryTableElement.Type.CsvStatAverage:
+					return "[csv]";
+				case SummaryTableElement.Type.CsvMetadata:
+					return "[meta]";
+				case SummaryTableElement.Type.ToolMetadata:
+					return "[toolmeta]";
+				default:
+					return "[invalid]";
+			}
+
+		}
+
+		public static SummaryTableElement.Type GetQualfiedStatType(string fullNameWithQualifier)
+		{
+			if (fullNameWithQualifier.StartsWith("["))
+			{
+				if (fullNameWithQualifier.StartsWith("[metric]"))
+				{
+					return SummaryTableElement.Type.SummaryTableMetric;
+				}
+				else if (fullNameWithQualifier.StartsWith("[csv]"))
+				{
+					return SummaryTableElement.Type.CsvStatAverage;
+				}
+				else if (fullNameWithQualifier.StartsWith("[meta]"))
+				{
+					return SummaryTableElement.Type.CsvMetadata;
+				}
+				else if (fullNameWithQualifier.StartsWith("[toolmeta]"))
+				{
+					return SummaryTableElement.Type.ToolMetadata;
+				}
+			}
+			return SummaryTableElement.Type.COUNT;
+		}
 
-			// Make a list of all unique keys
-			List<string> allMetadataKeys = new List<string>();
-			Dictionary<string, SummaryTableColumn> nameLookup = new Dictionary<string, SummaryTableColumn>();
-			foreach (SummaryTableColumn col in newTable.columns)
+		public static SummaryTableElement.Type GetQualfiedStatType(string fullNameWithQualifier, out string statNameOut)
+		{
+			if (fullNameWithQualifier.StartsWith("["))
 			{
-				string key = col.name.ToLower();
-				if (!nameLookup.ContainsKey(key))
+				int endIdx = fullNameWithQualifier.IndexOf("]");
+				if (endIdx == -1)
+				{
+					throw new Exception("Filter stat " + fullNameWithQualifier + " includes a start bracket but no end bracket");
+				}
+				// Strip off the stat qualifier
+				statNameOut = fullNameWithQualifier.Substring(endIdx+1);
+
+				if (fullNameWithQualifier.StartsWith("[metric]"))
+				{
+					return SummaryTableElement.Type.SummaryTableMetric;
+				}
+				else if (fullNameWithQualifier.StartsWith("[csv]"))
+				{
+					return SummaryTableElement.Type.CsvStatAverage;
+				}
+				else if (fullNameWithQualifier.StartsWith("[meta]"))
+				{
+					return SummaryTableElement.Type.CsvMetadata;
+				}
+				else if (fullNameWithQualifier.StartsWith("[toolmeta]"))
 				{
-					nameLookup.Add(key, col);
-					allMetadataKeys.Add(key);
+					return SummaryTableElement.Type.ToolMetadata;
 				}
 			}
-			allMetadataKeys.Sort();
+			statNameOut = fullNameWithQualifier;
+			return SummaryTableElement.Type.COUNT;
+		}
+
+		public SummaryTable SortAndFilter(List<string> columnFilterList, List<string> rowSortList, bool bReverseSort, string weightByColumnName, IEnumerable<ISummaryTableColumnFilter> additionalFilters = null)
+		{
+			SummaryTable newTable = SortRows(rowSortList, bReverseSort);
 
-			// Generate the list of requested metadata keys that this table includes
-			List<string> orderedKeysWithDupes = new List<string>();
+			// Generate a column lookup 
+			SummaryTableColumnLookup columnLookup = newTable.GetColumnLookup();
 
-			// Add metadata keys from the column filter list in the order they appear
+			// Make a list of all unique keys for each element type, including ANY
+			List<string>[] allMetadataKeysLists = new List<string>[(int)SummaryTableElement.Type.COUNT + 1];
+			for (int i = 0; i < allMetadataKeysLists.Length; i++)
+			{
+				allMetadataKeysLists[i] = columnLookup.GetSortedKeyList((SummaryTableElement.Type)i);
+			}
+
+			// Add columns from the column filter list in the order they appear
+			List<SummaryTableColumn> newColumnList = new List<SummaryTableColumn>();
 			foreach (string filterStr in columnFilterList)
 			{
 				string filterStrLower = filterStr.Trim().ToLower();
+				// Check for a qualifier which specifies the stat type
 				bool startWild = filterStrLower.StartsWith("*");
 				bool endWild = filterStrLower.EndsWith("*");
-				filterStrLower = filterStrLower.Trim('*');
-				if (startWild && endWild)
-                {
-					orderedKeysWithDupes.AddRange(allMetadataKeys.Where(x => x.Contains(filterStrLower)));
-                }
-				else if(startWild)
-				{
-					orderedKeysWithDupes.AddRange(allMetadataKeys.Where(x => x.EndsWith(filterStrLower)));
-				}
-				else if(endWild)
-				{
-					// Linear search through the sorted key list
-					bool bFound = false;
-					for (int wildcardSearchIndex = 0; wildcardSearchIndex < allMetadataKeys.Count; wildcardSearchIndex++)
+				if (startWild || endWild)
+				{
+					// Use the qualified list for wildcard matching if this entry was qualified with a type
+					SummaryTableElement.Type statType = GetQualfiedStatType(filterStrLower, out string unqualifiedKey);
+					List<string> keyList = new List<string>();
+					List<string> allKeysList = allMetadataKeysLists[(int)statType];
+					unqualifiedKey = unqualifiedKey.Trim('*');
+					if (startWild && endWild)
 					{
-						if (allMetadataKeys[wildcardSearchIndex].StartsWith(filterStrLower))
-						{
-							orderedKeysWithDupes.Add(allMetadataKeys[wildcardSearchIndex]);
-							bFound = true;
-						}
-						else if (bFound)
+						keyList.AddRange(allKeysList.Where(x => x.Contains(unqualifiedKey)));
+					}
+					else if (startWild)
+					{
+						keyList.AddRange(allKeysList.Where(x => x.EndsWith(unqualifiedKey)));
+					}
+					else if (endWild)
+					{
+						keyList.AddRange(allKeysList.Where(x => x.StartsWith(unqualifiedKey)));
+					}
+					// Resolve the keyList and output the columns 
+					foreach (string key in keyList)
+					{
+						SummaryTableColumn column = columnLookup.Get(key, statType);
+						if (column != null)
 						{
-							// Early exit: already found one key. If the pattern no longer matches then we must be done
-							break;
+							newColumnList.Add(column);
 						}
 					}
 				}
 				else
 				{
-					string key = filterStrLower;
-					orderedKeysWithDupes.Add(key);
+					SummaryTableColumn column = columnLookup.Get(filterStrLower);
+					if (column != null)
+					{
+						newColumnList.Add(column);
+					}
 				}
+
 			}
+			// Remove duplicates
+			newColumnList = newColumnList.Distinct().ToList();
 
 			// Compute row weights
-			if (weightByColumnName != null && nameLookup.ContainsKey(weightByColumnName))
+			if (weightByColumnName != null && columnLookup.ContainsKey(weightByColumnName))
 			{
-				SummaryTableColumn rowWeightColumn = nameLookup[weightByColumnName];
+				SummaryTableColumn rowWeightColumn = columnLookup.Get(weightByColumnName);
 				newTable.rowWeightings = new List<double>(rowWeightColumn.GetCount());
 				for (int i = 0; i < rowWeightColumn.GetCount(); i++)
 				{
@@ -1314,46 +1529,15 @@ namespace PerfSummaries
 				}
 			}
 
-			List<SummaryTableColumn> newColumnList = new List<SummaryTableColumn>();
-			// Add all the ordered keys that exist, ignoring duplicates
-			foreach (string key in orderedKeysWithDupes)
-			{
-				if (nameLookup.ContainsKey(key))
-				{
-					newColumnList.Add(nameLookup[key]);
-					// Remove from the list so it doesn't get counted again
-					nameLookup.Remove(key);
-				}
-			}
-
 			// Run the additional filters on the columns
 			if (additionalFilters != null)
 			{
 				var filteredColumns = new HashSet<string>();
 				foreach (ISummaryTableColumnFilter filter in additionalFilters)
 				{
-					if (showFilteredColumns)
-					{
-						// Run the filter so we can populate which columns would have been filtered, but don't actually remove them from the list.
-						newColumnList.ForEach(column =>
-						{
-							// If it's already filtered then skip it so we don't overwrite the filter reason.
-							if (!filteredColumns.Contains(column.name))
-							{
-								if (filter.ShouldFilter(column, this))
-								{
-									filteredColumns.Add(column.name);
-								}
-							}
-						});
-					}
-					else
-					{
-						newColumnList = newColumnList.Where(column => !filter.ShouldFilter(column, this)).ToList();
-					}
+					newColumnList = newColumnList.Where(column => !filter.ShouldFilter(column, this)).ToList();
 				}
 			}
-
 			newTable.columns = newColumnList;
 			newTable.rowCount = rowCount;
 			newTable.InitColumnLookup();
@@ -1833,14 +2017,14 @@ namespace PerfSummaries
 						{
 							// Work out the section name if we have section boundary info. When it changes, apply the sectionStart CSS class
 							string sectionName = "";
-							if (sectionBoundaryInfo != null && columnLookup.ContainsKey(sectionBoundaryInfo.statName))
+							if (sectionBoundaryInfo != null && columnLookup.ContainsKey(sectionBoundaryInfo.statName)) 
 							{
 								// Get the section name
 								if (!columnLookup.ContainsKey(sectionBoundaryInfo.statName))
 								{
 									continue;
 								}
-								SummaryTableColumn col = columnLookup[sectionBoundaryInfo.statName];
+								SummaryTableColumn col = columnLookup.Get(sectionBoundaryInfo.statName);
 								sectionName = col.GetStringValue(i);
 
 								// if we have a start token then strip before it
@@ -2040,9 +2224,9 @@ namespace PerfSummaries
 				string key = "";
 				foreach (string s in rowSortList)
 				{
-					if (columnLookup.ContainsKey(s.ToLower()))
+					SummaryTableColumn column = columnLookup.Get(s.ToLower());
+					if (column != null)
 					{
-						SummaryTableColumn column = columnLookup[s.ToLower()];
 						key += "{" + column.GetStringValue(i,false,"0000000000.0000000000") + "}";
 					}
 					else
@@ -2058,7 +2242,7 @@ namespace PerfSummaries
 				return m1.Key.CompareTo(m2.Key);
 			});
 
-			// Reorder the metadata rows
+			// Reorder the columns
 			List<SummaryTableColumn> newColumns = new List<SummaryTableColumn>();
 			foreach (SummaryTableColumn srcCol in columns)
 			{
@@ -2121,7 +2305,7 @@ namespace PerfSummaries
 				}
 				else
 				{
-					column = columnLookup[key];
+					column = columnLookup.Get(key);
 				}
 
 				if (value.isNumeric)
@@ -2145,14 +2329,10 @@ namespace PerfSummaries
 
 		public SummaryTableColumn GetColumnByName(string name)
 		{
-			if (columnLookup.ContainsKey(name))
-			{
-				return columnLookup[name];
-			}
-			return null;
+			return columnLookup.Get(name);
 		}
 
-		Dictionary<string, SummaryTableColumn> columnLookup = new Dictionary<string, SummaryTableColumn>();
+		SummaryTableColumnLookup columnLookup = new SummaryTableColumnLookup();
 		List<SummaryTableColumn> columns = new List<SummaryTableColumn>();
 		List<double> rowWeightings = null;
 		int rowCount = 0;
diff --git a/Engine/Source/Programs/CSVTools/PerfReportTool/SummaryTableData.cs b/Engine/Source/Programs/CSVTools/PerfReportTool/SummaryTableData.cs
--- a/Engine/Source/Programs/CSVTools/PerfReportTool/SummaryTableData.cs
+++ b/Engine/Source/Programs/CSVTools/PerfReportTool/SummaryTableData.cs
@@ -105,7 +105,9 @@ namespace PerfSummaries
 			CsvStatAverage,
 			CsvMetadata,
 			SummaryTableMetric,
-			ToolMetadata
+			ToolMetadata,
+			COUNT,
+			ANY=COUNT				
 		};
 
 		public enum Flags
@@ -521,7 +523,6 @@ namespace PerfSummaries
 
 		public SummaryTableRowData(Dictionary<string, dynamic> jsonDict)
 		{
-			var DataTypes = Enum.GetValues(typeof(SummaryTableElement.Type));
 			foreach (string dataTypeStr in jsonDict.Keys)
 			{
 				if ( Enum.TryParse(dataTypeStr, out SummaryTableElement.Type dataType) )
@@ -550,7 +551,10 @@ namespace PerfSummaries
 				var DataTypes = Enum.GetValues(typeof(SummaryTableElement.Type));
 				foreach (SummaryTableElement.Type dataType in DataTypes)
 				{
-					DictOut[dataType.ToString()] = new Dictionary<string, dynamic>();
+					if (dataType != SummaryTableElement.Type.COUNT)
+					{
+						DictOut[dataType.ToString()] = new Dictionary<string, dynamic>();
+					}
 				}
 			}
 
