PerfReportTool 4.233.0 : Unify wildcard handling for summary tables so it correctly matches mid-string wildcards. Also optimize using a binary search

diff --git a/Engine/Source/Programs/CSVTools/CsvStats/CsvStats.cs b/Engine/Source/Programs/CSVTools/CsvStats/CsvStats.cs
--- a/Engine/Source/Programs/CSVTools/CsvStats/CsvStats.cs
+++ b/Engine/Source/Programs/CSVTools/CsvStats/CsvStats.cs
@@ -636,6 +636,109 @@ namespace CSVStats
 			}
 		}
 
+		public static List<string> WildcardMatchStringList(List<string> stringList, string searchString, bool bCaseSensitive = false, bool bListIsSorted = false)
+		{
+			searchString = searchString.Trim();
+			if (!bCaseSensitive)
+			{
+				searchString = searchString.ToLower();
+			}
+			int startIndex = 0;
+			int endIndex = stringList.Count;
+
+			string[] wildcardSearchSubStrings = null;
+			if (searchString.Contains('*'))
+			{
+				wildcardSearchSubStrings = searchString.Split('*');
+			}
+
+			// If the list is sorted and there's no leading wildcard, we can binary search to limit the search range
+			if (bListIsSorted && !searchString.StartsWith('*'))
+			{
+				// In the wildcard case we're searching for the first term
+				string searchPrefix = wildcardSearchSubStrings == null ? searchString : wildcardSearchSubStrings[0];
+				int searchEndIndex = endIndex;
+
+				// Binary search for the first matching element
+				int prevMidIndex = -1;
+				while (true)
+				{
+					int midIndex = (startIndex + searchEndIndex) / 2;
+					if ( prevMidIndex == midIndex )
+					{
+						// No matches! Early out
+						return new List<string>();
+					}
+					string thisStr = bCaseSensitive ? stringList[midIndex] : stringList[midIndex].ToLower();
+					if (thisStr.StartsWith(searchPrefix))
+					{
+						string prevStr = midIndex > 0 ? (bCaseSensitive ? stringList[midIndex - 1] : stringList[midIndex - 1].ToLower()) : null;
+						if (prevStr == null || !prevStr.StartsWith(searchPrefix) )
+						{
+							// Exact find. Early out
+							startIndex = midIndex;
+							break;
+						}
+						else
+						{
+							searchEndIndex = midIndex;
+						}
+					}
+					else if (searchPrefix.CompareTo(thisStr) < 0)
+					{
+						searchEndIndex = midIndex;
+					}
+					else
+					{
+						startIndex = midIndex;
+					}
+					prevMidIndex = midIndex;
+				}
+
+
+				// Linear search to find the end of the region
+				for ( int i=startIndex+1; i<endIndex; i++ )
+				{
+					string thisStr = bCaseSensitive ? stringList[i] : stringList[i].ToLower();
+					if (!thisStr.StartsWith(searchPrefix))
+					{
+						endIndex = i;
+						break;
+					}
+				}
+
+				// We don't need a subsequent search if there's no wildcard or this is a single RHS wildcard
+				if (wildcardSearchSubStrings == null || (wildcardSearchSubStrings.Length == 2 && wildcardSearchSubStrings[1].Length == 0))
+				{
+					return stringList.GetRange(startIndex, endIndex-startIndex);
+				}
+			}
+
+
+			List<string> results = new List<string>();
+			if (wildcardSearchSubStrings != null)
+			{
+				for (int i = startIndex; i<endIndex; i++)
+				{
+					if (WildcardSubstringMatch(bCaseSensitive ? stringList[i] : stringList[i].ToLower(), wildcardSearchSubStrings))
+					{
+						results.Add(stringList[i]);
+					}
+				}
+			}
+			else
+			{
+				for (int i = startIndex; i < endIndex; i++)
+				{
+					if ((bCaseSensitive ? stringList[i] : stringList[i].ToLower()) == searchString)
+					{
+						results.Add(stringList[i]);
+					}
+				}
+			}
+			return results;
+		}
+
 
 		public Dictionary<string, bool> GetStatNamesMatchingStringList_Dict(string[] statNames)
 		{
@@ -668,20 +771,26 @@ namespace CSVStats
 				wildcardSearchSubStrings = statString.Split('*');
 			}
 			List<StatSamples> statList = new List<StatSamples>();
-			foreach (StatSamples stat in Stats.Values)
+			if (wildcardSearchSubStrings != null)
 			{
-				string statNameLower = stat.Name.ToLower();
-				if (wildcardSearchSubStrings != null)
+				foreach (StatSamples stat in Stats.Values)
 				{
-					if ( WildcardSubstringMatch(statNameLower, wildcardSearchSubStrings) )
+					if (WildcardSubstringMatch(stat.Name.ToLower(), wildcardSearchSubStrings))
 					{
 						statList.Add(stat);
 					}
 				}
-				else if (statNameLower == statString)
+			}
+			else
+			{
+				foreach (StatSamples stat in Stats.Values)
 				{
-					statList.Add(stat);
+					if (stat.Name.ToLower() == statString)
+					{
+						statList.Add(stat);
+					}
 				}
+
 			}
 			return statList;
 		}
diff --git a/Engine/Source/Programs/CSVTools/PerfReportTool/PerfReportTool.cs b/Engine/Source/Programs/CSVTools/PerfReportTool/PerfReportTool.cs
--- a/Engine/Source/Programs/CSVTools/PerfReportTool/PerfReportTool.cs
+++ b/Engine/Source/Programs/CSVTools/PerfReportTool/PerfReportTool.cs
@@ -23,7 +23,7 @@ namespace PerfReportTool
     class Version
     {
 		// Format: Major.Minor.Bugfix
-        private static string VersionString = "4.232.0";
+        private static string VersionString = "4.233.0";
 
         public static string Get() { return VersionString; }
     };
diff --git a/Engine/Source/Programs/CSVTools/PerfReportTool/SummaryTable.cs b/Engine/Source/Programs/CSVTools/PerfReportTool/SummaryTable.cs
--- a/Engine/Source/Programs/CSVTools/PerfReportTool/SummaryTable.cs
+++ b/Engine/Source/Programs/CSVTools/PerfReportTool/SummaryTable.cs
@@ -1473,6 +1473,7 @@ namespace PerfSummaries
 			foreach (string filterStr in columnFilterList)
 			{
 				string filterStrLower = filterStr.Trim().ToLower();
+
 				// Check for a qualifier which specifies the stat type
 				bool startWild = filterStrLower.StartsWith("*");
 				bool endWild = filterStrLower.EndsWith("*");
@@ -1480,21 +1481,10 @@ namespace PerfSummaries
 				{
 					// Use the qualified list for wildcard matching if this entry was qualified with a type
 					SummaryTableElement.Type statType = GetQualfiedStatType(filterStrLower, out string unqualifiedKey);
-					List<string> keyList = new List<string>();
+
 					List<string> allKeysList = allMetadataKeysLists[(int)statType];
-					unqualifiedKey = unqualifiedKey.Trim('*');
-					if (startWild && endWild)
-					{
-						keyList.AddRange(allKeysList.Where(x => x.Contains(unqualifiedKey)));
-					}
-					else if (startWild)
-					{
-						keyList.AddRange(allKeysList.Where(x => x.EndsWith(unqualifiedKey)));
-					}
-					else if (endWild)
-					{
-						keyList.AddRange(allKeysList.Where(x => x.StartsWith(unqualifiedKey)));
-					}
+					List<string> keyList = CsvStats.WildcardMatchStringList(allKeysList, unqualifiedKey, false, true);
+
 					// Resolve the keyList and output the columns 
 					foreach (string key in keyList)
 					{
