CSVToSVG 3.61
- Add -startEvent and -endEvent params to start and stop graphing at a particular event
- The params -startEventOffset and -endEventOffset allow you to offset the start/end by a particular number of frames
- If multiple CSVs are being processed, the data is cropped so frame counts start from zero
- The existing -maxx and -maxy parameters are relative to the adjusted CSV (e.g so -maxx=200 with -startEvent gc will capture 200 frames from the first gc)
- Fix bug with legend sorting. Frame is now always top
- Fix bug with auto Y axis scaling when a range is specified. Use the range instead of the whole CSV

Common
- Fix commandline parsing for negative numbers. Negative numbers are always recognized as values rather than param names

diff --git a/Engine/Source/Programs/CSVTools/CSVToSVG/CsvToSvgProgram.cs b/Engine/Source/Programs/CSVTools/CSVToSVG/CsvToSvgProgram.cs
--- a/Engine/Source/Programs/CSVTools/CSVToSVG/CsvToSvgProgram.cs
+++ b/Engine/Source/Programs/CSVTools/CSVToSVG/CsvToSvgProgram.cs
@@ -71,6 +71,10 @@ namespace CSVTools
 			"       -legend <list> \n" +
 			"       -maxHierarchyDepth <depth>\n" +
 			"       -minX <value> -maxX <value> -minY <value> -maxY <value>\n" +
+			"       -startEvent <name>\n" +
+			"       -startEventOffset <value>\n" +
+			"       -endEvent <name>\n" +
+			"       -endEventOffset <value>\n" +
 			"       -maxAutoMaxY <value> - clamp automatic maxY to this\n" +
 			"       -noMetadata\n" +
 			"       -noSnap\n" +
@@ -326,6 +330,11 @@ namespace CSVTools
 			graphParams.showEventNameText = GetIntArg("hideEventNames", 0) == 0;
 			graphParams.highlightEventRegions = GetListArg("highlightEventRegions", ',');
 
+			// Start/end event
+			graphParams.startEvent = GetArg("startEvent", null);
+			graphParams.startEventOffset = GetIntArg("startEventOffset", 0);
+			graphParams.endEvent = GetArg("endEvent", null);
+			graphParams.endEventOffset = GetIntArg("endEventOffset", 0);
 
 			// Smoothing
 			graphParams.smooth = GetArg("smooth") == "1";
diff --git a/Engine/Source/Programs/CSVTools/CsvStats/CommandLineTool.cs b/Engine/Source/Programs/CSVTools/CsvStats/CommandLineTool.cs
--- a/Engine/Source/Programs/CSVTools/CsvStats/CommandLineTool.cs
+++ b/Engine/Source/Programs/CSVTools/CsvStats/CommandLineTool.cs
@@ -22,6 +22,23 @@ namespace CSVStats
 			return commandLine;
 		}
 
+		public static bool IsParamName(string str)
+		{
+			if (str.Length == 0)
+			{
+				return false;
+			}
+			if (str[0] != '-')
+			{
+				return false;
+			}
+			if ( Double.TryParse(str, System.Globalization.NumberStyles.Any, System.Globalization.NumberFormatInfo.InvariantInfo, out _) )
+			{
+				return false;
+			}
+			return true;
+		}
+
 		private List<string> GetArgList(string inCommandLine)
 		{
 			List<string> args = new List<string>();
@@ -98,18 +115,18 @@ namespace CSVStats
 			for (int i = 0; i < args.Length; i++)
 			{
 				string arg = args[i];
-				if (arg[0] == '-')
+				if (IsParamName(arg))
 				{
 					string val = "1";
 
 					// If there's a value, read it
-					if (i < args.Length - 1 && args[i + 1][0] != '-')
+					if (i < args.Length - 1 && !IsParamName(args[i+1]))
 					{
 						bool first = true;
 						for (int j = i + 1; j < args.Length; j++)
 						{
 							string str = args[j];
-							if (str.Length > 0 && str[0] == '-')
+							if (str.Length > 0 && IsParamName(str))
 								break;
 							if (first)
 							{
diff --git a/Engine/Source/Programs/CSVTools/CsvStats/CsvStats.cs b/Engine/Source/Programs/CSVTools/CsvStats/CsvStats.cs
--- a/Engine/Source/Programs/CSVTools/CsvStats/CsvStats.cs
+++ b/Engine/Source/Programs/CSVTools/CsvStats/CsvStats.cs
@@ -415,6 +415,10 @@ namespace CSVStats
             {
                 maxSample = samples.Count;
             }
+			else
+			{
+				maxSample = Math.Min(maxSample, samples.Count);
+			}
             float maxValue = -float.MaxValue;
             for (int i = minSample; i < maxSample; i++)
             {
diff --git a/Engine/Source/Programs/CSVTools/CsvToSvgLib/CsvToSvgLib.cs b/Engine/Source/Programs/CSVTools/CsvToSvgLib/CsvToSvgLib.cs
--- a/Engine/Source/Programs/CSVTools/CsvToSvgLib/CsvToSvgLib.cs
+++ b/Engine/Source/Programs/CSVTools/CsvToSvgLib/CsvToSvgLib.cs
@@ -7,13 +7,15 @@ using System.Text;
 using System.ComponentModel;
 using System.Diagnostics;
 using System.Security.Cryptography;
+using System.Xml.Linq;
+using System.Collections;
 
 namespace CSVTools
 {
 
 	public class CsvToSvgLibVersion
 	{
-		private static string VersionString = "3.60";
+		private static string VersionString = "3.61";
 
 		public static string Get() { return VersionString; }
 	};
@@ -68,6 +70,12 @@ namespace CSVTools
 		public float minY = Range.Auto;
 		public float maxY = Range.Auto;
 
+		// start graphing from this event. Note: this will cause frame numbers to start at zero (which is necessary for multiple graphs)
+		public string startEvent;
+		public int startEventOffset = 0;
+		public string endEvent;
+		public int endEventOffset = 0;
+
 		// Max auto Y range. Set to 0 to disable
 		public float maxAutoMaxY = 0.0f;
 
@@ -476,21 +484,50 @@ namespace CSVTools
 			List<CsvStats> csvStatsList = new List<CsvStats>();
 			int currentColorOffset = graphParams.colorOffset;
 			int currentCustomLabelIndex = 0;
-			foreach (CsvInfo csvInfo in csvList)
+
+			FrameRange frameRange = new FrameRange();
+			for ( int i=0; i<csvList.Count; i++ )
 			{
-				CsvStats newCsvStats = ProcessCsvStats(csvInfo.stats, graphParams);
+				CsvInfo csvInfo = csvList[i];
+				FrameRange csvFrameRange;
+				CsvStats newCsvStats = ProcessCsvStats(csvInfo.stats, graphParams, out csvFrameRange);
+
+				if ( csvFrameRange.isLimited() )
+				{
+					// If we have multiple CSVs and we're applying startEvent limiting then just truncate the start of the CSVs at the event so the frames line up
+					// This has the effect that frame numbers start from zero, but there is no good alternative given the CSV frame numbers will differ
+					if (csvList.Count > 1 )
+					{
+						if (csvFrameRange.start > 0)
+						{
+							newCsvStats.CropStats(csvFrameRange.start, int.MaxValue);
+						}
+						// Adjust the global frameRange end such that we're using the largest range for all CSVs
+						if (csvFrameRange.end != int.MaxValue)
+						{
+							// Offset end index because we're truncating the start
+							csvFrameRange.end -= csvFrameRange.start;
+						}
+						frameRange.end = (i == 0) ? csvFrameRange.end : Math.Max(frameRange.end, csvFrameRange.end);
+					}
+					// Otherwise set frameRange to use the first (only) csv's range and we'll modify minX/maxX, which keeps the axis numbers accurate
+					else
+					{
+						frameRange = csvFrameRange;
+					}
+				}
 
 				if (graphParams.stacked && stackTotalStat == "")
 				{
 					// Make a total stat by summing each frame
 					StatSamples totalStat = new StatSamples("Total");
 					totalStat.samples.Capacity = newCsvStats.SampleCount;
-					for (int i = 0; i < newCsvStats.SampleCount; i++)
+					for (int j = 0; j < newCsvStats.SampleCount; j++)
 					{
 						float totalValue = 0.0f;
 						foreach (StatSamples stat in newCsvStats.Stats.Values)
 						{
-							totalValue += stat.samples[i];
+							totalValue += stat.samples[j];
 						}
 						totalStat.samples.Add(totalValue);
 					}
@@ -510,6 +547,22 @@ namespace CSVTools
 			}
 			perfLog.LogTiming("ProcessCsvStats");
 
+			Range range = new Range(graphParams.minX, graphParams.maxX, graphParams.minY, graphParams.maxY);
+
+			// Apply the frameRange if necessary
+			if (frameRange.isLimited())
+			{
+				// MinX and MaxX are both relative to frameRange.start
+				range.MinX = (range.MinX == Range.Auto) ? frameRange.start : frameRange.start + range.MinX;
+				if (range.MaxX != Range.Auto)
+				{
+					range.MaxX += frameRange.start;
+				}
+				if (frameRange.end != int.MaxValue)
+				{
+					range.MaxX = (range.MaxX == Range.Auto) ? frameRange.end : Math.Min(range.MaxX, frameRange.end); 
+				}
+			}
 
 			if (graphParams.smooth)
 			{
@@ -529,8 +582,6 @@ namespace CSVTools
 				perfLog.LogTiming("SmoothStats");
 			}
 
-			Range range = new Range(graphParams.minX, graphParams.maxX, graphParams.minY, graphParams.maxY);
-
 			// Compute the X range 
 			range = ComputeAdjustedXRange(range, graphRect, csvStatsList);
 
@@ -741,7 +792,7 @@ namespace CSVTools
 				{
 					foreach (StatSamples samples in stats.Stats.Values)
 					{
-						maxSample = Math.Max(maxSample, samples.ComputeMaxValue());
+						maxSample = Math.Max(maxSample, samples.ComputeMaxValue(range.MinX == Range.Auto ? 0 : (int)range.MinX, range.MaxX == Range.Auto ? -1 : (int)range.MaxX));
 					}
 				}
 				newRange.MaxY = Math.Min(maxSample * 1.05f, maxAutoMaxY);
@@ -774,10 +825,14 @@ namespace CSVTools
 			if (range.MinX == Range.Auto) newRange.MinX = 0;
 			if (range.MaxX == Range.Auto) newRange.MaxX = maxNumSamples;
 
-			// Quantize based on xincrement
+			// Quantize MinX and MaxX based on xincrement
 			if (rect != null)
 			{
 				float xInc = GetXAxisIncrement(rect, newRange);
+				if (newRange.MinX > 0)
+				{
+					newRange.MinX = (int)(range.MinX / xInc) * xInc;
+				}
 				float difX = newRange.MaxX - newRange.MinX;
 				float newDifX = (int)(0.9999 + difX / xInc) * xInc;
 				newRange.MaxX = newRange.MinX + newDifX;
@@ -1929,7 +1984,7 @@ namespace CSVTools
 			public StatSamples originalStatSamples;
 		};
 
-		void AddInteractiveScripting(SvgFile svg, Theme theme, Rect rect, Range range, List<CsvStats> csvStats, GraphParams graphParams)
+		void AddInteractiveScripting(SvgFile svg, Theme theme, Rect rect, Range range, List<CsvStats> csvStats, GraphParams graphParams )
 		{
 			bool bSnapToPeaks = graphParams.snapToPeaks && !graphParams.smooth;
 
@@ -2181,7 +2236,7 @@ namespace CSVTools
 
 
 
-			svg.WriteLine("function GetGraphX(mouseX)");
+			svg.WriteLine("function GetGraphX<UNIQUE>(mouseX)");
 			svg.WriteLine("{");
 			svg.WriteLine("  return (mouseX - " + rect.x + ") * (" + range.MaxX + " - " + range.MinX + ") / " + rect.width + " + " + range.MinX + ";");
 			svg.WriteLine("}");
@@ -2189,6 +2244,10 @@ namespace CSVTools
 
 			svg.WriteLine("function compareSamples(a, b)");
 			svg.WriteLine("{");
+			svg.WriteLine("      if (a.isFrame)");
+			svg.WriteLine("          return -1;");
+			svg.WriteLine("      if (b.isFrame)");
+			svg.WriteLine("          return 1;");
 			svg.WriteLine("      if (a.value > b.value)");
 			svg.WriteLine("          return -1;");
 			svg.WriteLine("      if (a.value < b.value)");
@@ -2196,7 +2255,7 @@ namespace CSVTools
 			svg.WriteLine("      return 0;");
 			svg.WriteLine("}");
 
-			svg.WriteLine("function ToSvgX(graphX)");
+			svg.WriteLine("function ToSvgX<UNIQUE>(graphX)");
 			svg.WriteLine("{");
 			svg.WriteLine("    scaleX = " + rect.width / (range.MaxX - range.MinX) + ";");
 			svg.WriteLine("    return " + rect.x + " + (graphX - " + range.MinX + ") * scaleX;");
@@ -2218,7 +2277,7 @@ namespace CSVTools
 
 			svg.WriteLine("function OnGraphAreaClicked<UNIQUE>(evt)");
 			svg.WriteLine("{");
-			svg.WriteLine("  graphX = GetGraphX(evt.offsetX); ");
+			svg.WriteLine("  graphX = GetGraphX<UNIQUE>(evt.offsetX); ");
 			svg.WriteLine("  var interactivePanel = document.getElementById('interactivePanel<UNIQUE>');");
 			svg.WriteLine("  var legendPanel = document.getElementById('LegendPanel<UNIQUE>');");
 			// Snap to an interesting frame (the max value under the pixel)
@@ -2227,7 +2286,7 @@ namespace CSVTools
 			svg.WriteLine("  if (frameNum >= " + range.MinX + " && frameNum < " + range.MaxX + ")");
 			svg.WriteLine("  {");
 			svg.WriteLine("    var xOffset = 0;");
-			svg.WriteLine("    var lineX = ToSvgX(frameNum);");
+			svg.WriteLine("    var lineX = ToSvgX<UNIQUE>(frameNum);");
 			svg.WriteLine("    var textX = lineX + xOffset;");
 			svg.WriteLine("    var textY = " + rect.y + " - 20");
 
@@ -2245,12 +2304,12 @@ namespace CSVTools
 				string textElementString = "document.getElementById('" + statInfo.jsTextElementId + "')";
 				if (statInfo.jsVarName == null)
 				{
-					svg.Write("            { value: frameNum, name: '" + statInfo.friendlyName + "', colour: 'rgb(0,0,0)'");
+					svg.Write("            { value: frameNum+" + graphParams.frameOffset + ", name: '" + statInfo.friendlyName + "', colour: 'rgb(0,0,0)', isFrame:true");
 				}
 				else
 				{
 					string valueStr = statInfo.jsVarName + "[dataIndex]";
-					svg.Write("            { value: " + valueStr + ", name: '" + statInfo.friendlyName + "', colour: " + statInfo.colour.SVGString());
+					svg.Write("            { value: " + valueStr + ", name: '" + statInfo.friendlyName + "', colour: " + statInfo.colour.SVGString() + ", isFrame:false");
 				}
 				svg.WriteLine(", groupElement: " + groupElementString + ", textElement: " + textElementString + " },");
 			}
@@ -2420,10 +2479,50 @@ namespace CSVTools
 			return false;
 		}
 
-		CsvStats ProcessCsvStats(CsvStats csvStatsIn, GraphParams graphParams)
+		
+		class FrameRange
 		{
-			CsvStats csvStats = new CsvStats(csvStatsIn, graphParams.statNames.ToArray());
+			public bool isLimited()
+			{
+				return start > 0 || end < Int32.MaxValue;
+			}
+			
+			public int start = 0;
+			public int end = Int32.MaxValue;
+		};
+
+		FrameRange GetEventTruncationFrameRange(CsvStats csvStats, GraphParams graphParams)
+		{
+			// Apply startEvent and endEvent truncation if requested
+			FrameRange frameRange = new FrameRange();
+			if (graphParams.startEvent != null)
+			{
+				foreach (CsvEvent ev in csvStats.Events)
+				{
+					if (CsvStats.DoesSearchStringMatch(ev.Name, graphParams.startEvent))
+					{
+						frameRange.start = Math.Clamp(ev.Frame + graphParams.startEventOffset, 0, csvStats.SampleCount - 1);
+						break;
+					}
+				}
+			}
+			if (graphParams.endEvent != null)
+			{
+				foreach (CsvEvent ev in csvStats.Events)
+				{
+					if (ev.Frame >= frameRange.start && CsvStats.DoesSearchStringMatch(ev.Name, graphParams.endEvent))
+					{
+						frameRange.end = Math.Clamp(ev.Frame + graphParams.endEventOffset, 0, csvStats.SampleCount - 1);
+						break;
+					}
+				}
+			}
+			return frameRange;
+		}
 
+		CsvStats ProcessCsvStats(CsvStats csvStatsIn, GraphParams graphParams, out FrameRange frameRange)
+		{
+			CsvStats csvStats = new CsvStats(csvStatsIn, graphParams.statNames.ToArray());
 			if (graphParams.discardLastFrame)
 			{
 				foreach (StatSamples stat in csvStats.Stats.Values.ToArray())
@@ -2553,6 +2652,8 @@ namespace CSVTools
 				}
 			}
 
+			frameRange = GetEventTruncationFrameRange(csvStats, graphParams);
+
 			// Filter out events
 			List<CsvEvent> FilteredEvents = new List<CsvEvent>();
 			foreach (CsvEvent ev in csvStats.Events)
