PerfReportTool 4.234.0
- Add support for startEvent and stopEvent in graphs
- Add requiredEvents field for graphs so we can suppress graphs based on whether particular events are present

diff --git a/Engine/Source/Programs/CSVTools/PerfReportTool/PerfReportTool.cs b/Engine/Source/Programs/CSVTools/PerfReportTool/PerfReportTool.cs
--- a/Engine/Source/Programs/CSVTools/PerfReportTool/PerfReportTool.cs
+++ b/Engine/Source/Programs/CSVTools/PerfReportTool/PerfReportTool.cs
@@ -23,7 +23,7 @@ namespace PerfReportTool
     class Version
     {
 		// Format: Major.Minor.Bugfix
-        private static string VersionString = "4.233.0";
+        private static string VersionString = "4.234.0";
 
         public static string Get() { return VersionString; }
     };
@@ -1039,18 +1039,8 @@ namespace PerfReportTool
 						throw new Exception("Graph " + graph.title + " has no <statString> element");
 					}
 
-					bool bFoundStat = false;
-					foreach (string statString in graph.settings.statString.value.Split(','))
-					{
-						List<StatSamples> matchingStats = csvFile.dummyCsvStats.GetStatsMatchingString(statString);
-						if (matchingStats.Count > 0)
-						{
-							bFoundStat = true;
-							break;
-						}
-					}
-
-					if (bFoundStat)
+					bool bShowGraph = graph.ShouldShow(csvFile.finalCsv ?? csvFile.dummyCsvStats);
+					if (bShowGraph)
 					{
 						if (useEmbeddedGraphUrl)
 						{
@@ -1841,6 +1831,11 @@ namespace PerfReportTool
 			AddOptionalArg("legendAverageThreshold", graphSettings.legendAverageThreshold);
 			AddConditionalArg("ignoreStats", graphSettings.ignoreStats.isSet, graphSettings.ignoreStats.value);
 
+			AddConditionalArg("startEvent", graphSettings.startEvent.isSet, graphSettings.startEvent.value);
+			AddConditionalArg("startEventOffset", graphSettings.startEventOffset.isSet, graphSettings.startEventOffset.value);
+			AddConditionalArg("endEvent", graphSettings.endEvent.isSet, graphSettings.endEvent.value);
+			AddConditionalArg("endEventOffset", graphSettings.endEventOffset.isSet, graphSettings.endEventOffset.value);
+
 			string argString = string.Empty;
 			if (argFormat == CsvToSvgArgFormat.Url)
 			{
@@ -1965,6 +1960,25 @@ namespace PerfReportTool
 			{
 				graphParams.statMultiplier = (float)graphSettings.statMultiplier.value;
 			}
+
+			if (graphSettings.startEvent.isSet)
+			{
+				graphParams.startEvent = graphSettings.startEvent.value;
+				if (graphSettings.startEventOffset.isSet)
+				{
+					graphParams.startEventOffset = graphSettings.startEventOffset.value;
+				}					
+			}
+
+			if (graphSettings.endEvent.isSet)
+			{
+				graphParams.endEvent = graphSettings.endEvent.value;
+				if (graphSettings.endEventOffset.isSet)
+				{
+					graphParams.endEventOffset = graphSettings.endEventOffset.value;
+				}
+			}
+
 			graphParams.interactive = true;
 
 			if (!GetBoolArg("noStripEvents"))
diff --git a/Engine/Source/Programs/CSVTools/PerfReportTool/ReportGraph.cs b/Engine/Source/Programs/CSVTools/PerfReportTool/ReportGraph.cs
--- a/Engine/Source/Programs/CSVTools/PerfReportTool/ReportGraph.cs
+++ b/Engine/Source/Programs/CSVTools/PerfReportTool/ReportGraph.cs
@@ -1,6 +1,8 @@
 // Copyright Epic Games, Inc. All Rights Reserved.
 
+using CSVStats;
 using System;
+using System.Collections.Generic;
 using System.Xml.Linq;
 
 namespace PerfReportTool
@@ -111,6 +113,46 @@ namespace PerfReportTool
 				settings = new GraphSettings(element, vars);
 			}
 		}
+
+		public bool ShouldShow(CsvStats stats)
+		{
+			bool bFoundStat = false;
+			foreach (string statString in settings.statString.value.Split(','))
+			{
+				List<StatSamples> matchingStats = stats.GetStatsMatchingString(statString);
+				if (matchingStats.Count > 0)
+				{
+					bFoundStat = true;
+					break;
+				}
+			}
+			if ( !bFoundStat ) 
+			{
+				return false;
+			}
+
+			if (settings.requiredEvents.isSet)
+			{
+				string[] requiredEvents = settings.requiredEvents.value.Split(',');
+				foreach (string requiredEvent in requiredEvents)
+				{
+					bool bFoundEvent = false;
+					foreach (CsvEvent ev in stats.Events)
+					{
+						if (CsvStats.DoesSearchStringMatch(requiredEvent, ev.Name))
+						{
+							bFoundEvent = true;
+							break;
+						}
+					}
+					if (!bFoundEvent)
+					{
+						return false;
+					}
+				}
+			}
+			return true;
+		}
 		public string title;
 		public Optional<double> budget;
 		public bool inSummary;
@@ -155,6 +197,13 @@ namespace PerfReportTool
 			legendAverageThreshold = new Optional<double>(element, "legendAverageThreshold", vars);
 			snapToPeaks = new Optional<bool>(element, "snapToPeaks", vars);
 			lineDecimalPlaces = new Optional<int>(element, "lineDecimalPlaces", vars);
+
+			startEvent = new OptionalString(element, "startEvent", false, vars);
+			startEventOffset = new Optional<int>(element, "startEventOffset", vars);
+			endEvent = new OptionalString(element, "endEvent", false, vars);
+			endEventOffset = new Optional<int>(element, "endEventOffset", vars);
+
+			requiredEvents = new OptionalString(element, "requiredEvents", false, vars);
 		}
 		public void InheritFrom(GraphSettings baseSettings)
 		{
@@ -187,7 +236,11 @@ namespace PerfReportTool
 			legendAverageThreshold.InheritFrom(baseSettings.legendAverageThreshold);
 			snapToPeaks.InheritFrom(baseSettings.snapToPeaks);
 			lineDecimalPlaces.InheritFrom(baseSettings.lineDecimalPlaces);
-
+			startEvent.InheritFrom(baseSettings.startEvent);
+			startEventOffset.InheritFrom(baseSettings.startEventOffset);
+			endEvent.InheritFrom(baseSettings.endEvent);
+			endEventOffset.InheritFrom(baseSettings.endEventOffset);
+			requiredEvents.InheritFrom(baseSettings.requiredEvents);
 		}
 		public Optional<bool> smooth;
 		public OptionalString statString;
@@ -220,6 +273,13 @@ namespace PerfReportTool
 		public Optional<bool> snapToPeaks;
 		public Optional<int> lineDecimalPlaces;
 
+		public OptionalString startEvent;
+		public Optional<int> startEventOffset;
+		public OptionalString endEvent;
+		public Optional<int> endEventOffset;
+
+		public OptionalString requiredEvents;
+
 	};
 
 }
\ No newline at end of file
