PerfReportTool 4.235.0. Add hideInDetailedReports attribute to summaries allowing them to be suppressed. This is useful if you only want to emit summary table row data for a given summary

diff --git a/Engine/Source/Programs/CSVTools/PerfReportTool/PerfReportTool.cs b/Engine/Source/Programs/CSVTools/PerfReportTool/PerfReportTool.cs
--- a/Engine/Source/Programs/CSVTools/PerfReportTool/PerfReportTool.cs
+++ b/Engine/Source/Programs/CSVTools/PerfReportTool/PerfReportTool.cs
@@ -23,7 +23,7 @@ namespace PerfReportTool
     class Version
     {
 		// Format: Major.Minor.Bugfix
-        private static string VersionString = "4.234.0";
+        private static string VersionString = "4.235.0";
 
         public static string Get() { return VersionString; }
     };
@@ -1495,7 +1495,8 @@ namespace PerfReportTool
 			// If the reporttype has summary info, then write out the summary]
 			foreach (Summary summary in summaries)
 			{
-				HtmlSection htmlSection = summary.WriteSummaryData(htmlFile != null, summary.useUnstrippedCsvStats ? csvStatsUnstripped : csvStats, csvStatsUnstripped, bWriteSummaryCsv, summaryRowData, htmlFilename);
+				bool bWriteHtml = htmlFile != null && !summary.bHideInDetailedReport;
+				HtmlSection htmlSection = summary.WriteSummaryData(bWriteHtml, summary.useUnstrippedCsvStats ? csvStatsUnstripped : csvStats, csvStatsUnstripped, bWriteSummaryCsv, summaryRowData, htmlFilename);
 				if (htmlSection != null)
 				{
 					htmlSection.WriteToFile(htmlFile);
@@ -1676,7 +1677,8 @@ namespace PerfReportTool
 			// If the reporttype has summary info, then write out the summary]
 			foreach (Summary summary in reportTypeInfo.summaries)
 			{
-				HtmlSection htmlSection = summary.WriteSummaryData(htmlFile != null, csvStats, csvStatsUnstripped, bWriteSummaryCsv, null, htmlFilename);
+				bool bWriteHtml = htmlFile != null && !summary.bHideInDetailedReport;
+				HtmlSection htmlSection = summary.WriteSummaryData(bWriteHtml, csvStats, csvStatsUnstripped, bWriteSummaryCsv, null, htmlFilename);
 				if (htmlSection != null)
 				{
 					htmlSection.WriteToFile(htmlFile);
diff --git a/Engine/Source/Programs/CSVTools/PerfReportTool/Summary.cs b/Engine/Source/Programs/CSVTools/PerfReportTool/Summary.cs
--- a/Engine/Source/Programs/CSVTools/PerfReportTool/Summary.cs
+++ b/Engine/Source/Programs/CSVTools/PerfReportTool/Summary.cs
@@ -78,8 +78,7 @@ namespace PerfSummaries
             stats = new List<string>();
             captures = new List<CaptureRange>();
             StatThresholds = new Dictionary<string, ColourThresholdList>();
-
-        }
+		}
 		public virtual HtmlSection WriteSummaryData(bool bWriteHtml, CsvStats csvStats, CsvStats csvStatsUnstripped, bool bWriteSummaryCsv, SummaryTableRowData rowData, string htmlFileName)
 		{
 			return null;
@@ -99,6 +98,7 @@ namespace PerfSummaries
 			{
 				return;
 			}
+			bHideInDetailedReport = element.GetSafeAttribute<bool>(vars, "hideInDetailedReport", false);
 			useUnstrippedCsvStats = element.GetSafeAttribute<bool>(vars, "useUnstrippedCsvStats", false);
 			bStartCollapsed = element.GetSafeAttribute<bool>(vars, "collapsed", false);
 			XElement statsElement = element.Element("stats");
@@ -242,6 +242,7 @@ namespace PerfSummaries
         public Dictionary<string, ColourThresholdList> StatThresholds;
 		public bool useUnstrippedCsvStats;
 		public bool bStartCollapsed;
+		public bool bHideInDetailedReport = false;
     };
 
 
