CsvCollate 1.35 - add StartEvent param so we can trim the start of the CSVs we're combining. This is useful for syncing up timelines

diff --git a/Engine/Source/Programs/CSVTools/CSVCollate/CSVCollate.cs b/Engine/Source/Programs/CSVTools/CSVCollate/CSVCollate.cs
--- a/Engine/Source/Programs/CSVTools/CSVCollate/CSVCollate.cs
+++ b/Engine/Source/Programs/CSVTools/CSVCollate/CSVCollate.cs
@@ -14,7 +14,7 @@ namespace CSVTools
 {
     class Version
     {
-        private static string VersionString = "1.34";
+        private static string VersionString = "1.35";
 
         public static string Get() { return VersionString; }
     };
@@ -30,6 +30,7 @@ namespace CSVTools
 			"       [-filterOutlierStat <stat>] - discard CSVs if this stat has very high values\n" +
 			"       [-filterOutlierThreshold <value>] - threshold for outliers (default:1000)\n" +
 			"       [-metadataFilter <key=value,key=value...>] : filters based on CSV metadata\n" +
+			"       [-startEvent <event name>] : starts from a particular event\n" +
 			"       -o <csvFilename> \n";
 
 		void Run(string[] args)
@@ -128,7 +129,7 @@ namespace CSVTools
 
 			CsvStats combinedCsvStats = new CsvStats();
 
-
+			string startEventName = GetArg("startEvent", null);
 			string metadataFilterString = GetArg("metadataFilter", null);
 			List<int> frameCsvCounts=new List<int>();
 			List<string> allCsvFilenames = new List<string>();
@@ -137,6 +138,23 @@ namespace CSVTools
             {
                 CsvStats srcCsvStats = CsvStats.ReadCSVFile(csvFilename, null);
 
+				if (startEventName != null)
+				{
+					int startFrame = 0;
+					foreach (CsvEvent ev in srcCsvStats.Events)
+					{
+						if (CsvStats.DoesSearchStringMatch(ev.Name, startEventName))
+						{
+							startFrame = ev.Frame;
+							break;
+						}
+					}
+					if (startFrame != 0 )
+					{
+						srcCsvStats.CropStats(startFrame);
+					}
+				}
+
 				// Check for outliers
 				bool skip = false;
 				if ( filterOutlierStat != null )
