PerfReportTool v4.236.0 - Added Wildcard Stat support to BoundedStatSummary. This allows ReportTypes.xml to specify '*' for generated columns in a summary table. The tool will now attempt to match that wildcard to all possible stats, generating new columns for each. Wildcard symbol needs to be provided in the column name, summaryState name and stat name, or it will skip that stat name. Supports otherStat as well so can be used with the Ratio formula.
#tests Built in Release and tested both standalone and via PRS
[RN] PerfReportTool v4.236.0 Added Wildcard Stat support to BoundedStatSummary. This allows ReportTypes.xml to specify '*' for generated columns in a summary table

diff --git a/Engine/Source/Programs/CSVTools/PerfReportTool/PerfReportTool.cs b/Engine/Source/Programs/CSVTools/PerfReportTool/PerfReportTool.cs
--- a/Engine/Source/Programs/CSVTools/PerfReportTool/PerfReportTool.cs
+++ b/Engine/Source/Programs/CSVTools/PerfReportTool/PerfReportTool.cs
@@ -23,7 +23,7 @@ namespace PerfReportTool
     class Version
     {
 		// Format: Major.Minor.Bugfix
-        private static string VersionString = "4.235.0";
+        private static string VersionString = "4.236.0";
 
         public static string Get() { return VersionString; }
     };
diff --git a/Engine/Source/Programs/CSVTools/PerfReportTool/Summaries/BoundedStatValuesSummary.cs b/Engine/Source/Programs/CSVTools/PerfReportTool/Summaries/BoundedStatValuesSummary.cs
--- a/Engine/Source/Programs/CSVTools/PerfReportTool/Summaries/BoundedStatValuesSummary.cs
+++ b/Engine/Source/Programs/CSVTools/PerfReportTool/Summaries/BoundedStatValuesSummary.cs
@@ -27,6 +27,26 @@ namespace PerfSummaries
 			public double frameExponent; // Exponent for relative frame time (0-1) in streamingstressmetric formula
 			public double statExponent; // Exponent for stat value in streamingstressmetric formula
 			public ColourThresholdList colourThresholdList;
+
+			public Column() { }
+
+			public Column(Column inCol)
+			{
+				name = inCol.name;
+				formula = inCol.formula;
+				value = inCol.value;
+				summaryStatName = inCol.summaryStatName;
+				statName = inCol.statName;
+				otherStatName = inCol.otherStatName;
+				perSecond = inCol.perSecond;
+				filterOutZeros = inCol.filterOutZeros;
+				applyEndOffset = inCol.applyEndOffset;
+				multiplier = inCol.multiplier;
+				threshold = inCol.threshold;
+				frameExponent = inCol.frameExponent;
+				statExponent = inCol.statExponent;
+				colourThresholdList = inCol.colourThresholdList;
+			}
 		};
 		public BoundedStatValuesSummary(XElement element, XmlVariableMappings vars, string baseXmlDirectory)
 		{
@@ -142,14 +162,95 @@ namespace PerfSummaries
 			StatSamples frameTimeStat = csvStats.GetStat("frametime");
 			List<float> frameTimes = frameTimeStat.samples;
 
-			// Filter only columns with stats that exist in the CSV
 			List<Column> filteredColumns = new List<Column>();
+
+			//Filtering and wildcard logic
+			//This adds new columns for each stat which matches a wildcard column, and also filters the columns to ones
+			//that have a valid stat associated with them. The original wildcard column will be removed as well
 			foreach (Column col in columns)
 			{
-				if (csvStats.GetStat(col.statName) != null
-					&& (String.IsNullOrWhiteSpace(col.otherStatName) || csvStats.GetStat(col.otherStatName) != null))
+				bool nameIsWildcard = col.name.Contains('*');
+				bool summaryStatNameIsWildcard = col.summaryStatName.Contains('*');
+				bool statNameIsWildcard = col.statName.Contains('*');
+				bool otherStatNameIsSet = col.otherStatName.Length > 0;
+				bool otherStatNameIsWildcard = col.otherStatName.Contains('*');
+				bool anyNameWildcard = nameIsWildcard || summaryStatNameIsWildcard || statNameIsWildcard || otherStatNameIsWildcard;
+				bool allNameWildcard = nameIsWildcard && summaryStatNameIsWildcard && statNameIsWildcard && (otherStatNameIsSet ? otherStatNameIsWildcard : true);
+
+				//Check all wildcard fields are set, or none of them should be set
+				if (anyNameWildcard && !allNameWildcard)
+				{
+					string errorString = "Warning: BoundedStatValuesSummary: Skipping column because wildcard * was found in some column parameters but not all. name: " + col.name + " summaryStatName: " + col.summaryStatName + " statName: " + col.statName;
+					if(otherStatNameIsSet)
+					{
+						errorString += " otherStatName: " + col.otherStatName;
+					}
+					Console.WriteLine(errorString);
+
+					continue;
+				}
+
+				//Generate new columns based on the wildcard definition
+				if (anyNameWildcard && allNameWildcard)
+				{
+					string[] splitStr = col.statName.Split('*', 2);
+					if(splitStr.Length == 0)
+					{
+						Console.WriteLine("Warning: Skipping stat due to error: StatName is a wildcard and should contain a * symbol: " + col.statName);
+						continue;
+					}
+
+					//Find each stat that matches the wildcard in statName
+					string statPrefix = splitStr[0];
+					foreach (KeyValuePair<string, StatSamples> foundStat in csvStats.Stats)
+					{
+						string foundStatName = foundStat.Key;
+						if (foundStatName.ToLower().Contains(statPrefix.ToLower()) && !foundStatName.Equals(col.statName))
+						{
+							//Create a new column as a copy of the Wildcard Definition Column
+							Column newCol = new Column(col);
+
+							string[] splitString = foundStatName.Split(statPrefix, 2);
+							if(splitString.Length <= 1)
+							{
+								Console.WriteLine("Warning: Skipping stat due to error: FoundStatName(" + foundStatName + ") does not contain expected prefix: " + statPrefix);
+								continue;
+							}
+							
+							string statPostfix = splitString[1];
+							newCol.name = col.name.Substring(0, col.name.IndexOf('*')) + statPostfix;
+							newCol.summaryStatName = col.summaryStatName.Substring(0, col.summaryStatName.IndexOf('*')) + statPostfix;
+							newCol.statName = col.statName.Substring(0, col.statName.IndexOf('*')) + statPostfix;
+							if(csvStats.GetStat(newCol.statName) == null)
+							{
+								Console.WriteLine("Warning: Skipping stat due to error: StatName not found in csvstats: " + newCol.statName);
+								continue;
+							}
+							
+							if (otherStatNameIsSet)
+							{
+								newCol.otherStatName = col.otherStatName.Substring(0, col.otherStatName.IndexOf('*')) + statPostfix;
+								if (csvStats.GetStat(newCol.otherStatName) == null)
+								{
+									Console.WriteLine("Warning: Skipping stat due to error: StatName not found in csvstats: " + newCol.otherStatName);
+									continue;
+								}
+							}
+
+							//Add newly formed column to the filteredColumns list
+							filteredColumns.Add(newCol);							
+						}
+					}
+				}
+				else
 				{
-					filteredColumns.Add(col);
+					// This is a non-wildcard column, so add it as usual
+					// Filter only columns with stats that exist in the CSV
+					if (csvStats.GetStat(col.statName) != null
+						&& (String.IsNullOrWhiteSpace(col.otherStatName) || csvStats.GetStat(col.otherStatName) != null))
+					{
+						filteredColumns.Add(col); 
+					}
 				}
 			}
 
