PRT - Fix null object exception in bounded summary stat table.

#rnx

diff --git a/Engine/Source/Programs/CSVTools/PerfReportTool/Summaries/BoundedStatValuesSummary.cs b/Engine/Source/Programs/CSVTools/PerfReportTool/Summaries/BoundedStatValuesSummary.cs
--- a/Engine/Source/Programs/CSVTools/PerfReportTool/Summaries/BoundedStatValuesSummary.cs
+++ b/Engine/Source/Programs/CSVTools/PerfReportTool/Summaries/BoundedStatValuesSummary.cs
@@ -169,11 +169,11 @@ namespace PerfSummaries
 			//that have a valid stat associated with them. The original wildcard column will be removed as well
 			foreach (Column col in columns)
 			{
-				bool nameIsWildcard = col.name.Contains('*');
-				bool summaryStatNameIsWildcard = col.summaryStatName.Contains('*');
-				bool statNameIsWildcard = col.statName.Contains('*');
-				bool otherStatNameIsSet = col.otherStatName.Length > 0;
-				bool otherStatNameIsWildcard = col.otherStatName.Contains('*');
+				bool nameIsWildcard = col.name != null && col.name.Contains('*');
+				bool summaryStatNameIsWildcard = col.summaryStatName != null && col.summaryStatName.Contains('*');
+				bool statNameIsWildcard = col.statName != null && col.statName.Contains('*');
+				bool otherStatNameIsSet = col.otherStatName != null && col.otherStatName.Length > 0;
+				bool otherStatNameIsWildcard = col.otherStatName != null && col.otherStatName.Contains('*');
 				bool anyNameWildcard = nameIsWildcard || summaryStatNameIsWildcard || statNameIsWildcard || otherStatNameIsWildcard;
 				bool allNameWildcard = nameIsWildcard && summaryStatNameIsWildcard && statNameIsWildcard && (otherStatNameIsSet ? otherStatNameIsWildcard : true);
 
