PRT - Guard against divide by 0 in BoundedStatValuesSummary

#rnx

diff --git a/Engine/Source/Programs/CSVTools/PerfReportTool/PerfReportTool.cs b/Engine/Source/Programs/CSVTools/PerfReportTool/PerfReportTool.cs
--- a/Engine/Source/Programs/CSVTools/PerfReportTool/PerfReportTool.cs
+++ b/Engine/Source/Programs/CSVTools/PerfReportTool/PerfReportTool.cs
@@ -23,7 +23,7 @@ namespace PerfReportTool
     class Version
     {
 		// Format: Major.Minor.Bugfix
-        private static string VersionString = "4.236.0";
+        private static string VersionString = "4.236.1";
 
         public static string Get() { return VersionString; }
     };
diff --git a/Engine/Source/Programs/CSVTools/PerfReportTool/Summaries/BoundedStatValuesSummary.cs b/Engine/Source/Programs/CSVTools/PerfReportTool/Summaries/BoundedStatValuesSummary.cs
--- a/Engine/Source/Programs/CSVTools/PerfReportTool/Summaries/BoundedStatValuesSummary.cs
+++ b/Engine/Source/Programs/CSVTools/PerfReportTool/Summaries/BoundedStatValuesSummary.cs
@@ -448,7 +448,8 @@ namespace PerfSummaries
 					{
 						denominator += otherStatValues[i];
 					}
-					value = numerator / denominator; // TODO: Does the rest of the pipeline handle +/-i infinity?
+					// Guard against divide by 0 as json serialization cannot handle Inf.
+					value = denominator != 0.0 ? numerator / denominator : 0.0;
 					totalFrameWeight = 1.0;
 				}
 				else
@@ -456,7 +457,7 @@ namespace PerfSummaries
 					throw new Exception("BoundedStatValuesSummary: unexpected formula " + col.formula);
 				}
 				value *= col.multiplier;
-				col.value = value / totalFrameWeight;
+				col.value = totalFrameWeight != 0.0 ? value / totalFrameWeight : value;
 			}
 
 			// Output HTML
