PerfReportTool 4.237.0 - Make summary table sorting smarter. If a column with string values includes an integer suffix, left pad with zeroes for sorting purposes

diff --git a/Engine/Source/Programs/CSVTools/PerfReportTool/PerfReportTool.cs b/Engine/Source/Programs/CSVTools/PerfReportTool/PerfReportTool.cs
--- a/Engine/Source/Programs/CSVTools/PerfReportTool/PerfReportTool.cs
+++ b/Engine/Source/Programs/CSVTools/PerfReportTool/PerfReportTool.cs
@@ -23,7 +23,7 @@ namespace PerfReportTool
     class Version
     {
 		// Format: Major.Minor.Bugfix
-        private static string VersionString = "4.236.1";
+        private static string VersionString = "4.237.0";
 
         public static string Get() { return VersionString; }
     };
diff --git a/Engine/Source/Programs/CSVTools/PerfReportTool/SummaryTable.cs b/Engine/Source/Programs/CSVTools/PerfReportTool/SummaryTable.cs
--- a/Engine/Source/Programs/CSVTools/PerfReportTool/SummaryTable.cs
+++ b/Engine/Source/Programs/CSVTools/PerfReportTool/SummaryTable.cs
@@ -2217,7 +2217,27 @@ namespace PerfSummaries
 					SummaryTableColumn column = columnLookup.Get(s.ToLower());
 					if (column != null)
 					{
-						key += "{" + column.GetStringValue(i,false,"0000000000.0000000000") + "}";
+						string columnKey = column.GetStringValue(i, false, "0000000000.0000000000");
+						if (!column.isNumeric)
+						{
+							// If there's an integer suffix in the column value, pad it with zeroes for sorting purposes
+							int integerSuffixStartIndex = -1;
+							for (int ci = columnKey.Length-1; ci>0; ci-- )
+							{
+								char c = columnKey[ci];
+								if (c < '0' || c > '9')
+								{
+									break;
+								}
+								integerSuffixStartIndex = ci;
+							}
+							if (integerSuffixStartIndex >= 0)
+							{
+								string integerSuffixPadded = columnKey.Substring(integerSuffixStartIndex).PadLeft(12,'0');
+								columnKey = columnKey.Substring(0, integerSuffixStartIndex) + integerSuffixPadded;
+							}
+						}
+						key += "{" + columnKey + "}";
 					}
 					else
 					{
