PRT - Write error logs to stderr instead of stdout.

#rnx

diff --git a/Engine/Source/Programs/CSVTools/CSVCollate/CSVCollate.cs b/Engine/Source/Programs/CSVTools/CSVCollate/CSVCollate.cs
--- a/Engine/Source/Programs/CSVTools/CSVCollate/CSVCollate.cs
+++ b/Engine/Source/Programs/CSVTools/CSVCollate/CSVCollate.cs
@@ -285,7 +285,7 @@ namespace CSVTools
                 }
                 catch (System.Exception e)
                 {
-                    Console.WriteLine("[ERROR] " + e.Message);
+                    Console.Error.WriteLine("[ERROR] " + e.Message);
                 }
             }
         }
diff --git a/Engine/Source/Programs/CSVTools/CSVFilter/CSVFilter.cs b/Engine/Source/Programs/CSVTools/CSVFilter/CSVFilter.cs
--- a/Engine/Source/Programs/CSVTools/CSVFilter/CSVFilter.cs
+++ b/Engine/Source/Programs/CSVTools/CSVFilter/CSVFilter.cs
@@ -209,7 +209,7 @@ namespace CSVTools
                 }
                 catch (System.Exception e)
                 {
-                    Console.WriteLine("[ERROR] " + e.Message);
+                    Console.Error.WriteLine("[ERROR] " + e.Message);
                 }
             }
         }
diff --git a/Engine/Source/Programs/CSVTools/CSVSplit/CSVSplit.cs b/Engine/Source/Programs/CSVTools/CSVSplit/CSVSplit.cs
--- a/Engine/Source/Programs/CSVTools/CSVSplit/CSVSplit.cs
+++ b/Engine/Source/Programs/CSVTools/CSVSplit/CSVSplit.cs
@@ -179,7 +179,7 @@ namespace CSVTools
                 }
                 catch (System.Exception e)
                 {
-                    Console.WriteLine("[ERROR] " + e.Message);
+                    Console.Error.WriteLine("[ERROR] " + e.Message);
                 }
             }
         }
diff --git a/Engine/Source/Programs/CSVTools/CSVToSVG/CsvToSvgProgram.cs b/Engine/Source/Programs/CSVTools/CSVToSVG/CsvToSvgProgram.cs
--- a/Engine/Source/Programs/CSVTools/CSVToSVG/CsvToSvgProgram.cs
+++ b/Engine/Source/Programs/CSVTools/CSVToSVG/CsvToSvgProgram.cs
@@ -488,7 +488,7 @@ namespace CSVTools
 				}
 				catch (System.Exception e)
 				{
-					Console.WriteLine("[ERROR] " + e.Message);
+					Console.Error.WriteLine("[ERROR] " + e.Message);
 					if (Debugger.IsAttached)
 					{
 						throw;
diff --git a/Engine/Source/Programs/CSVTools/CsvConvert/CsvConvert.cs b/Engine/Source/Programs/CSVTools/CsvConvert/CsvConvert.cs
--- a/Engine/Source/Programs/CSVTools/CsvConvert/CsvConvert.cs
+++ b/Engine/Source/Programs/CSVTools/CsvConvert/CsvConvert.cs
@@ -35,7 +35,7 @@ namespace CSVInfo
                 }
                 catch (System.Exception e)
                 {
-                    Console.WriteLine("[ERROR] " + e.Message);
+                    Console.Error.WriteLine("[ERROR] " + e.Message);
                     return 1;
                 }
             }
diff --git a/Engine/Source/Programs/CSVTools/CsvInfo/CSVInfo.cs b/Engine/Source/Programs/CSVTools/CsvInfo/CSVInfo.cs
--- a/Engine/Source/Programs/CSVTools/CsvInfo/CSVInfo.cs
+++ b/Engine/Source/Programs/CSVTools/CsvInfo/CSVInfo.cs
@@ -34,7 +34,7 @@ namespace CSVInfo
                 }
                 catch (System.Exception e)
                 {
-                    Console.WriteLine("[ERROR] " + e.Message);
+                    Console.Error.WriteLine("[ERROR] " + e.Message);
                     return 1;
                 }
             }
diff --git a/Engine/Source/Programs/CSVTools/PerfReportTool/ColumnFilters/RegressionColumnFilter.cs b/Engine/Source/Programs/CSVTools/PerfReportTool/ColumnFilters/RegressionColumnFilter.cs
--- a/Engine/Source/Programs/CSVTools/PerfReportTool/ColumnFilters/RegressionColumnFilter.cs
+++ b/Engine/Source/Programs/CSVTools/PerfReportTool/ColumnFilters/RegressionColumnFilter.cs
@@ -168,7 +168,7 @@ namespace PerfReportTool
 			SummaryTableColumn groupByColumn = table.GetColumnByName(JoinRowsByStatName);
 			if (groupByColumn == null)
 			{
-				Console.WriteLine($"[Error] -regressionJoinRowsByName doesn't refer to a real column name: {JoinRowsByStatName}. Will not be applying any grouping to rows.");
+				Console.Error.WriteLine($"[Error] -regressionJoinRowsByName doesn't refer to a real column name: {JoinRowsByStatName}. Will not be applying any grouping to rows.");
 				return 1;
 			}
 
diff --git a/Engine/Source/Programs/CSVTools/PerfReportTool/PerfReportTool.cs b/Engine/Source/Programs/CSVTools/PerfReportTool/PerfReportTool.cs
--- a/Engine/Source/Programs/CSVTools/PerfReportTool/PerfReportTool.cs
+++ b/Engine/Source/Programs/CSVTools/PerfReportTool/PerfReportTool.cs
@@ -23,7 +23,7 @@ namespace PerfReportTool
     class Version
     {
 		// Format: Major.Minor.Bugfix
-        private static string VersionString = "4.237.0";
+        private static string VersionString = "4.237.1";
 
         public static string Get() { return VersionString; }
     };
@@ -629,7 +629,7 @@ namespace PerfReportTool
 				{
 					if (bBulkMode)
 					{
-						Console.Out.WriteLine("[ERROR] : " + e.Message);
+						Console.Error.WriteLine("[ERROR] : " + e.Message);
 					}
 					else
 					{
@@ -2108,7 +2108,7 @@ namespace PerfReportTool
 				}
 				catch (System.Exception e)
 				{
-					Console.WriteLine("[ERROR] " + e.Message);
+					Console.Error.WriteLine("[ERROR] " + e.Message);
 					return 1;
 				}
 			}
