Make padding trailing digits optional using the -sortTrailingDigitsAsNumeric param since it can cause issues

diff --git a/Engine/Source/Programs/CSVTools/PerfReportTool/PerfReportTool.cs b/Engine/Source/Programs/CSVTools/PerfReportTool/PerfReportTool.cs
--- a/Engine/Source/Programs/CSVTools/PerfReportTool/PerfReportTool.cs
+++ b/Engine/Source/Programs/CSVTools/PerfReportTool/PerfReportTool.cs
@@ -23,7 +23,7 @@ namespace PerfReportTool
     class Version
     {
 		// Format: Major.Minor.Bugfix
-        private static string VersionString = "4.237.1";
+        private static string VersionString = "4.237.2";
 
         public static string Get() { return VersionString; }
     };
@@ -147,6 +147,7 @@ namespace PerfReportTool
 			"  -spreadsheetfriendly: outputs a single quote before non-numeric entries in summary tables\n" +
 			"  -noSummaryMinMax: don't make min/max columns for each stat in a condensed summary\n" +
 			"  -reverseTable [0|1]: Reverses the order of summary tables (set 0 to force off)\n" +
+			"  -sortTrailingDigitsAsNumeric : detects trailing digits and pads them when sorting summary table\n"+
 			"  -scrollableTable [0|1]: makes the summary table scrollable, with frozen first rows and columns (set 0 to force off)\n" +
 			"  -colorizeTable [off|budget|auto]: selects the table colorization mode. If omitted, uses the default in the summary\n" +
 			"     xml table if set.\n" +
@@ -799,7 +800,8 @@ namespace PerfReportTool
 			// TODO: would be better if we could determine HighIsBad without the format info and store it directly in the column.
 			table.SetColumnFormatInfo(reportXML.columnFormatInfoList);
 
-			SummaryTable filteredTable = table.SortAndFilter(tableInfo.columnFilterList, tableInfo.rowSortList, bReverseTable, weightByColumnName, additionalColumnFilters);
+			bool bSortTrailingDigitsAsNumeric = GetBoolArg("sortTrailingDigitsAsNumeric");
+			SummaryTable filteredTable = table.SortAndFilter(tableInfo.columnFilterList, tableInfo.rowSortList, bReverseTable, weightByColumnName, additionalColumnFilters, bSortTrailingDigitsAsNumeric);
 			if (bCollated)
 			{
 				filteredTable = filteredTable.CollateSortedTable(tableInfo.rowSortList, addMinMaxColumns);
diff --git a/Engine/Source/Programs/CSVTools/PerfReportTool/SummaryTable.cs b/Engine/Source/Programs/CSVTools/PerfReportTool/SummaryTable.cs
--- a/Engine/Source/Programs/CSVTools/PerfReportTool/SummaryTable.cs
+++ b/Engine/Source/Programs/CSVTools/PerfReportTool/SummaryTable.cs
@@ -1372,13 +1372,6 @@ namespace PerfSummaries
 			return columnLookup;
 		}
 
-
-		public SummaryTable SortAndFilter(string customFilter, string customRowSort = "buildversion,deviceprofile", bool bReverseSort = false, string weightByColumnName = null)
-		{
-			return SortAndFilter(customFilter.Split(',').ToList(), customRowSort.Split(',').ToList(), bReverseSort, weightByColumnName);
-		}
-
-
 		public static string GetElementTypeStatPrefix(SummaryTableElement.Type type)
 		{
 			switch (type)
@@ -1454,9 +1447,9 @@ namespace PerfSummaries
 			return SummaryTableElement.Type.COUNT;
 		}
 
-		public SummaryTable SortAndFilter(List<string> columnFilterList, List<string> rowSortList, bool bReverseSort, string weightByColumnName, IEnumerable<ISummaryTableColumnFilter> additionalFilters = null)
+		public SummaryTable SortAndFilter(List<string> columnFilterList, List<string> rowSortList, bool bReverseSort, string weightByColumnName, IEnumerable<ISummaryTableColumnFilter> additionalFilters = null, bool bSortTrailingDigitsAsNumeric = false)
 		{
-			SummaryTable newTable = SortRows(rowSortList, bReverseSort);
+			SummaryTable newTable = SortRows(rowSortList, bReverseSort, bSortTrailingDigitsAsNumeric);
 
 			// Generate a column lookup 
 			SummaryTableColumnLookup columnLookup = newTable.GetColumnLookup();
@@ -2206,7 +2199,7 @@ namespace PerfSummaries
 			htmlFile.Close();
 		}
 
-		public SummaryTable SortRows(List<string> rowSortList, bool reverseSort)
+		public SummaryTable SortRows(List<string> rowSortList, bool reverseSort, bool sortTrailingDigitsAsNumeric = false)
 		{
 			List<KeyValuePair<string, int>> columnRemapping = new List<KeyValuePair<string, int>>();
 			for (int i = 0; i < rowCount; i++)
@@ -2218,7 +2211,7 @@ namespace PerfSummaries
 					if (column != null)
 					{
 						string columnKey = column.GetStringValue(i, false, "0000000000.0000000000");
-						if (!column.isNumeric)
+						if (sortTrailingDigitsAsNumeric && !column.isNumeric)
 						{
 							// If there's an integer suffix in the column value, pad it with zeroes for sorting purposes
 							int integerSuffixStartIndex = -1;
