PerfReportTool 4.237.3: Variable Resolution and graceful failure improvements/bugfixes
- Fix a bug in variable resolution for array variables where the whole field would be replaced instead of just the variable
- Replace variables which fail to resolve with an empty string instead of leaving the variable name in place. This is easier for calling code to detect and handle
- Add a graceful failures in a few places to allow us to generate a report if variables fail to resolve

diff --git a/Engine/Source/Programs/CSVTools/PerfReportTool/PerfReportTool.cs b/Engine/Source/Programs/CSVTools/PerfReportTool/PerfReportTool.cs
--- a/Engine/Source/Programs/CSVTools/PerfReportTool/PerfReportTool.cs
+++ b/Engine/Source/Programs/CSVTools/PerfReportTool/PerfReportTool.cs
@@ -23,7 +23,7 @@ namespace PerfReportTool
     class Version
     {
 		// Format: Major.Minor.Bugfix
-        private static string VersionString = "4.237.2";
+        private static string VersionString = "4.237.3";
 
         public static string Get() { return VersionString; }
     };
diff --git a/Engine/Source/Programs/CSVTools/PerfReportTool/ReportXml.cs b/Engine/Source/Programs/CSVTools/PerfReportTool/ReportXml.cs
--- a/Engine/Source/Programs/CSVTools/PerfReportTool/ReportXml.cs
+++ b/Engine/Source/Programs/CSVTools/PerfReportTool/ReportXml.cs
@@ -624,13 +624,12 @@ namespace PerfReportTool
 
 		public string ResolveVariables(string attributeValue)
 		{
-			// Remap all variables found in the attribute name
 			if (!attributeValue.Contains('$'))
 			{
 				return attributeValue;
 			}
 
-			// Remap all variables found in the attribute name
+			// Remap all variables found in the attribute value
 			int StringPos = 0;
 			while (StringPos < attributeValue.Length)
 			{
@@ -650,12 +649,14 @@ namespace PerfReportTool
 
 				string FullVariableName = attributeValue.Substring(VarStartIndex+2, VarEndIndex - VarStartIndex-2);
 				string VariableName = FullVariableName;
-				int ArrayIndex = -1;
+
+				string VariableValue = "";
 
 				// Check for an array index
 				int OpenBracketIndex = VariableName.IndexOf('[');
 				if (OpenBracketIndex != -1)
 				{
+					int ArrayIndex = -1;
 					if (FullVariableName.EndsWith("]"))
 					{
 						string ArrayIndexStr = VariableName.Substring(OpenBracketIndex + 1, VariableName.Length - 2 - OpenBracketIndex);
@@ -667,33 +668,42 @@ namespace PerfReportTool
 					if (ArrayIndex < 0)
 					{
 						Console.WriteLine("[Warning] Failed to resolve variable ${" + FullVariableName + "}. Can't read array index");
-						continue;
 					}
-					VariableName = FullVariableName.Substring(0, OpenBracketIndex);
-				}
-
-
-				// Replace the variable if found
-				if (vars.TryGetValue(VariableName, out string VariableValue))
-				{
-					attributeValue = attributeValue.Substring(0,VarStartIndex) + VariableValue + attributeValue.Substring(VarEndIndex+1);
-					if (ArrayIndex >= 0)
+					else
 					{
-						string[] elements = attributeValue.Split(",");
-						if (ArrayIndex >= elements.Length)
+						VariableName = FullVariableName.Substring(0, OpenBracketIndex);
+						if (vars.TryGetValue(VariableName, out VariableValue))
+						{
+							string[] elements = VariableValue.Split(",");
+							if (ArrayIndex < elements.Length)
+							{
+								VariableValue = elements[ArrayIndex];
+							}
+							else
+							{
+								Console.WriteLine("[Warning] Failed to resolve variable ${" + FullVariableName + "}. Array index out of range!");
+							}
+						}
+						else
 						{
-							Console.WriteLine("[Warning] Failed to resolve variable ${" + FullVariableName + "}. Array index out of range!");
-							continue;
+							Console.WriteLine("[Warning] Failed to resolve array variable ${" + FullVariableName + "}");
+							VariableValue = "";
 						}
-						attributeValue = elements[ArrayIndex];
 					}
-					// Adjust stringPos to take into account the replace
-					StringPos = VarStartIndex + VariableValue.Length;
 				}
-				else
+				else 
 				{
-					Console.WriteLine("[Warning] Failed to resolve variable ${" + VariableName + "}");
+					// Read the variable value. Default to empty string and replace anyway if not found
+					if (!vars.TryGetValue(VariableName, out VariableValue))
+					{
+						Console.WriteLine("[Warning] Failed to resolve variable ${" + VariableName + "}");
+						VariableValue = "";
+					}
 				}
+				
+				// Replace the variable name with its value and update StringPos to take into account the replace
+				attributeValue = attributeValue.Substring(0, VarStartIndex) + VariableValue + attributeValue.Substring(VarEndIndex + 1);
+				StringPos = VarStartIndex + VariableValue.Length;
 			}
 			return attributeValue;
 		}
@@ -820,8 +830,10 @@ namespace PerfReportTool
 				}
 			}
 			catch (FormatException e)
-			{
-				Console.WriteLine(string.Format("[Warning] Failed to convert XML attribute '{0}' '{1}' ({2})", attributeName, attributeValue, e.Message));
+			{	
+				// If the attribute value is empty (likely due to failed variable mapping), display the original version
+				string attributeValueToDisplay = attributeValue.Length > 0 ? attributeValue : attribute.Value;
+				Console.WriteLine(string.Format("[Warning] Failed to convert XML attribute '{0}' '{1}' ({2})", attributeName, attributeValueToDisplay, e.Message));
 				return defaultValue;
 			}
 		}
diff --git a/Engine/Source/Programs/CSVTools/PerfReportTool/Summaries/HistogramSummary.cs b/Engine/Source/Programs/CSVTools/PerfReportTool/Summaries/HistogramSummary.cs
--- a/Engine/Source/Programs/CSVTools/PerfReportTool/Summaries/HistogramSummary.cs
+++ b/Engine/Source/Programs/CSVTools/PerfReportTool/Summaries/HistogramSummary.cs
@@ -21,7 +21,10 @@ namespace PerfSummaries
 			HistogramThresholds = new double[histogramStrings.Length];
 			for (int i = 0; i < histogramStrings.Length; i++)
 			{
-				HistogramThresholds[i] = Convert.ToDouble(histogramStrings[i], System.Globalization.CultureInfo.InvariantCulture);
+				if (histogramStrings[i].Length > 0) // Fail gracefully if a variable fails to resolve
+				{
+					HistogramThresholds[i] = Convert.ToDouble(histogramStrings[i], System.Globalization.CultureInfo.InvariantCulture);
+				}
 			}
 
 			foreach (XElement child in element.Elements())
diff --git a/Engine/Source/Programs/CSVTools/PerfReportTool/Summaries/HitchSummary.cs b/Engine/Source/Programs/CSVTools/PerfReportTool/Summaries/HitchSummary.cs
--- a/Engine/Source/Programs/CSVTools/PerfReportTool/Summaries/HitchSummary.cs
+++ b/Engine/Source/Programs/CSVTools/PerfReportTool/Summaries/HitchSummary.cs
@@ -20,7 +20,10 @@ namespace PerfSummaries
 			HitchThresholds = new double[hitchThresholds.Length];
 			for (int i = 0; i < hitchThresholds.Length; i++)
 			{
-				HitchThresholds[i] = Convert.ToDouble(hitchThresholds[i], System.Globalization.CultureInfo.InvariantCulture);
+				if (hitchThresholds[i].Length > 0) // Fail gracefully if a variable fails to resolve
+				{
+					HitchThresholds[i] = Convert.ToDouble(hitchThresholds[i], System.Globalization.CultureInfo.InvariantCulture);
+				}
 			}
 		}
 		public HitchSummary() { }
diff --git a/Engine/Source/Programs/CSVTools/PerfReportTool/SummaryTypes.cs b/Engine/Source/Programs/CSVTools/PerfReportTool/SummaryTypes.cs
--- a/Engine/Source/Programs/CSVTools/PerfReportTool/SummaryTypes.cs
+++ b/Engine/Source/Programs/CSVTools/PerfReportTool/SummaryTypes.cs
@@ -435,7 +435,8 @@ namespace PerfSummaries
 				string[] colourStrings = colourThresholdElementStr.Split(',');
 				if (colourStrings.Length != 4)
 				{
-					throw new Exception("Incorrect number of colourthreshold entries. Should be 4.");
+					Console.Write("Incorrect number of colourthreshold entries. Should be 4: "+ colourThresholdEl.ToString());
+					return null;
 				}
 				double[] colourThresholds = new double[4];
 				for (int i = 0; i < colourStrings.Length; i++)
