PRT - Switch error log to stderr

Missed this one in the original pass.

#rnx

diff --git a/Engine/Source/Programs/CSVTools/PerfReportTool/CsvFileCache..cs b/Engine/Source/Programs/CSVTools/PerfReportTool/CsvFileCache..cs
--- a/Engine/Source/Programs/CSVTools/PerfReportTool/CsvFileCache..cs
+++ b/Engine/Source/Programs/CSVTools/PerfReportTool/CsvFileCache..cs
@@ -434,7 +434,7 @@ namespace PerfReportTool
 			{
 				if (bBulkMode)
 				{
-					Console.Out.WriteLine("[ERROR] : " + e.Message);
+					Console.Error.WriteLine("[ERROR] : " + e.Message);
 					return null;
 				}
 				else
diff --git a/Engine/Source/Programs/CSVTools/PerfReportTool/PerfReportTool.cs b/Engine/Source/Programs/CSVTools/PerfReportTool/PerfReportTool.cs
--- a/Engine/Source/Programs/CSVTools/PerfReportTool/PerfReportTool.cs
+++ b/Engine/Source/Programs/CSVTools/PerfReportTool/PerfReportTool.cs
@@ -23,7 +23,7 @@ namespace PerfReportTool
     class Version
     {
 		// Format: Major.Minor.Bugfix
-        private static string VersionString = "4.237.3";
+        private static string VersionString = "4.237.4";
 
         public static string Get() { return VersionString; }
     };
