PRT (minor) Fix mapOverlaySummary logging

diff --git a/Engine/Source/Programs/CSVTools/PerfReportTool/Summaries/MapOverlaySummary.cs b/Engine/Source/Programs/CSVTools/PerfReportTool/Summaries/MapOverlaySummary.cs
--- a/Engine/Source/Programs/CSVTools/PerfReportTool/Summaries/MapOverlaySummary.cs
+++ b/Engine/Source/Programs/CSVTools/PerfReportTool/Summaries/MapOverlaySummary.cs
@@ -136,7 +136,7 @@ namespace PerfSummaries
 					var thumbnail = image.GetThumbnailImage(destWidth, destHeight, null, IntPtr.Zero);
 					using (var destImageStream = new FileStream(destImagePath, FileMode.OpenOrCreate, FileAccess.Write))
 					{
-						Console.WriteLine("Saving downsampled map image: " + destImageStream);
+						Console.WriteLine("Saving downsampled map image: " + destImagePath);
 						thumbnail.Save(destImageStream, ImageFormat.Jpeg);
 					}
 				}
