CsvToSvg 3.62 - Improve visual clutter due to too many events:
- Merge displayed event names if they're too close together, respecting priority (priority is implicit based on param order)
- Display nearby event names in an interactive panel alongside stats
- Truncate displayed events on the graph to 32 characters
Note: you can pass -showAllEvents to disable the new merging/truncation behaviour

diff --git a/Engine/Source/Programs/CSVTools/CSVToSVG/CsvToSvgProgram.cs b/Engine/Source/Programs/CSVTools/CSVToSVG/CsvToSvgProgram.cs
--- a/Engine/Source/Programs/CSVTools/CSVToSVG/CsvToSvgProgram.cs
+++ b/Engine/Source/Programs/CSVTools/CSVToSVG/CsvToSvgProgram.cs
@@ -63,6 +63,7 @@ namespace CSVTools
 			"       -fixedPointPrecisionScale <1..N> - scale for fixed point graph rendering (>1 gives subpixel accuracy)"+
 			"       -graphOnly\n" +
 			"       -hideEventNames <1|0>\n" +
+			"       -showAllEventNames <1|0> - don't filter out event names if they're too dense\n" +
 			"       -hideStatPrefix <list>\n" +
 			"       -hierarchySeparator <character>\n" +
 			"       -highlightEventRegions <startEventName,endEventName>\n" +
@@ -327,7 +328,14 @@ namespace CSVTools
 
 			// Events
 			graphParams.showEventNames = GetListArg("showEvents");
-			graphParams.showEventNameText = GetIntArg("hideEventNames", 0) == 0;
+			if ( GetIntArg("hideEventNames", 0) == 1 )
+			{
+				graphParams.showEventNameTextMode = ShowEventTextMode.Hide;
+			}
+			else if (GetIntArg("showAllEventNames", 0) == 1)
+			{
+				graphParams.showEventNameTextMode = ShowEventTextMode.ShowAll;
+			}
 			graphParams.highlightEventRegions = GetListArg("highlightEventRegions", ',');
 
 			// Start/end event
diff --git a/Engine/Source/Programs/CSVTools/CsvToSvgLib/CsvToSvgLib.cs b/Engine/Source/Programs/CSVTools/CsvToSvgLib/CsvToSvgLib.cs
--- a/Engine/Source/Programs/CSVTools/CsvToSvgLib/CsvToSvgLib.cs
+++ b/Engine/Source/Programs/CSVTools/CsvToSvgLib/CsvToSvgLib.cs
@@ -15,11 +15,18 @@ namespace CSVTools
 
 	public class CsvToSvgLibVersion
 	{
-		private static string VersionString = "3.61";
+		private static string VersionString = "3.62";
 
 		public static string Get() { return VersionString; }
 	};
 
+	public enum ShowEventTextMode
+	{
+		Hide = 0,
+		ShowAuto = 1,
+		ShowAll = 2
+	};
+
 	public class GraphParams
 	{
 		public int width = 1800;
@@ -31,7 +38,7 @@ namespace CSVTools
 
 		// Events
 		public List<string> showEventNames = new List<string>();
-		public bool showEventNameText = true;
+		public ShowEventTextMode showEventNameTextMode = ShowEventTextMode.ShowAuto;
 		public List<string> highlightEventRegions = new List<string>();
 
 		// Smoothing
@@ -724,10 +731,10 @@ namespace CSVTools
 						DrawGraph(svg, stat.samples, stat.colour, graphRect, range, thickness, statID, graphParams, MaxSegments);
 					}
 
-					if (graphParams.showEventNameText)
+					if (graphParams.showEventNameTextMode != ShowEventTextMode.Hide)
 					{
 						Colour eventColour = theme.EventTextColour;
-						DrawEventText(svg, csvStat.Events, eventColour, graphRect, range);
+						DrawEventText(svg, csvStat.Events, eventColour, graphRect, range, graphParams);
 					}
 
 					csvIndex++;
@@ -1632,37 +1639,51 @@ namespace CSVTools
 		}
 
 
-		class CsvEventWithCount : CsvEvent
+		class CsvGraphEvent : CsvEvent
 		{
-			public CsvEventWithCount(CsvEvent ev)
+			public CsvGraphEvent(CsvEvent ev, int inPriority)
 			{
 				base.Frame = ev.Frame;
 				base.Name = ev.Name;
 				count = 1;
+				priority = inPriority;
 			}
 			public int count;
+			public int priority;
 		};
 
-		void DrawEventText(SvgFile svg, List<CsvEvent> events, Colour colour, Rect rect, Range range)
+		private string truncateTextIfNeeded(string text, int maxLength)
+		{
+			if (text.Length <= maxLength)
+			{
+				return text;
+			}
+			return text.Substring(0, maxLength) + "...";
+		}
+
+		void DrawEventText(SvgFile svg, List<CsvEvent> events, Colour colour, Rect rect, Range range, GraphParams graphParams )
 		{
 			float LastEventX = -100000.0f;
 			int lastFrame = 0;
 
+			bool bAutoShowMode = graphParams.showEventNameTextMode == ShowEventTextMode.ShowAuto;
+			float mergeThreshold = bAutoShowMode ? 12.0f : 8.5f;
+
 			// Make a filtered list of events to display, grouping duplicates which are close together
-			List<CsvEventWithCount> filteredEvents = new List<CsvEventWithCount>();
-			CsvEventWithCount currentDisplayEvent = null;
-			CsvEvent lastEvent = null;
+			List<CsvGraphEvent> filteredEvents = new List<CsvGraphEvent>();
+			CsvGraphEvent currentDisplayEvent = null;
 			foreach (CsvEvent ev in events)
 			{
 				// Only draw events which are in the range
+				int eventPriority = GetEventPriority(ev.Name, graphParams);
 				if (ev.Frame >= range.MinX && ev.Frame <= range.MaxX)
 				{
 					// Merge with the current display event?
 					bool bMerge = false;
-					if (currentDisplayEvent != null && lastEvent != null && ev.Name == currentDisplayEvent.Name)
+					if (currentDisplayEvent != null && ( ev.Name == currentDisplayEvent.Name || bAutoShowMode))
 					{
-						float DistToLastEvent = ToSvgX(ev.Frame, rect, range) - ToSvgX(lastEvent.Frame, rect, range);
-						if (DistToLastEvent <= 8.5)
+						float DistToLastEvent = ToSvgX(ev.Frame, rect, range) - ToSvgX(currentDisplayEvent.Frame, rect, range);
+						if (DistToLastEvent <= mergeThreshold)
 						{
 							bMerge = true;
 						}
@@ -1670,44 +1691,73 @@ namespace CSVTools
 
 					if (bMerge)
 					{
+						if ( bAutoShowMode )
+						{
+							// If this event is actually higher (numerically lower) priority then replace the current event
+							if (eventPriority < currentDisplayEvent.priority )
+							{
+								currentDisplayEvent.priority = eventPriority;
+								currentDisplayEvent.Name = ev.Name;
+								currentDisplayEvent.Frame = ev.Frame;
+							}
+						}
 						currentDisplayEvent.count++;
 					}
 					else
 					{
-						currentDisplayEvent = new CsvEventWithCount(ev);
+						currentDisplayEvent = new CsvGraphEvent(ev, eventPriority);
 						filteredEvents.Add(currentDisplayEvent);
 					}
-					lastEvent = ev;
 				}
 			}
 
-			foreach (CsvEventWithCount ev in filteredEvents)
+			if (graphParams.showEventNameTextMode == ShowEventTextMode.ShowAuto)
 			{
-				float eventX = ToSvgX(ev.Frame, rect, range);
-				string name = ev.Name;
-
-				if (ev.count > 1)
+				foreach (CsvGraphEvent ev in filteredEvents)
 				{
-					name += " &#x00D7; " + ev.count;
-				}
+					float eventX = ToSvgX(ev.Frame, rect, range);
+					string name = truncateTextIfNeeded(ev.Name,32);
 
-				// Space out the events (allow at least 8 pixels between them)
-				if (eventX - LastEventX <= 8.5f)
-				{
-					// Add an arrow to indicate events were spaced out 
-					name = "&#x21b7; " + name;
-					if (ev.count == 1)
+					if (ev.count > 1)
 					{
-						name += " (+" + (ev.Frame - lastFrame) + ")";
+						name += " +" + (ev.count-1);
 					}
-					eventX = LastEventX + 9.0f;
+
+					float csvTextX = ToCsvX(eventX + 7, rect, range);
+					DrawHorizontalAxisText(svg, name, csvTextX, colour, rect, range, 10, true);
 				}
-				float csvTextX = ToCsvX(eventX + 7, rect, range);
+			}
+			else
+			{
+				foreach (CsvGraphEvent ev in filteredEvents)
+				{
+					float eventX = ToSvgX(ev.Frame, rect, range);
+					string name = ev.Name;
 
-				DrawHorizontalAxisText(svg, name, csvTextX, colour, rect, range, 10, true);
-				LastEventX = eventX;
-				lastFrame = ev.Frame;
+					if (ev.count > 1)
+					{
+						name += " &#x00D7; " + ev.count;
+					}
+
+					// Space out the events (allow at least 8 pixels between them)
+					if (eventX - LastEventX <= 8.5f)
+					{
+						// Add an arrow to indicate events were spaced out 
+						name = "&#x21b7; " + name;
+						if (ev.count == 1)
+						{
+							name += " (+" + (ev.Frame - lastFrame) + ")";
+						}
+						eventX = LastEventX + 9.0f;
+					}
+					float csvTextX = ToCsvX(eventX + 7, rect, range);
+
+					DrawHorizontalAxisText(svg, name, csvTextX, colour, rect, range, 10, true);
+					LastEventX = eventX;
+					lastFrame = ev.Frame;
+				}
 			}
+
 		}
 
 		void DrawEventLines(SvgFile svg, Theme theme, List<CsvEvent> events, Rect rect, Range range)
@@ -2022,6 +2072,7 @@ namespace CSVTools
 			// TODO: strip out data outside the range (needs an offset as well as a multiplier)
 			int multiplier = 1;
 			float numStatsPerPixel = (float)(range.MaxX - range.MinX) / rect.width;
+			int filteredSampleCount = (int)(range.MaxX) - (int)(range.MinX);
 
 			if (numStatsPerPixel > 1)
 			{
@@ -2057,7 +2108,7 @@ namespace CSVTools
 					}
 				}
 
-				int filteredStatCount = maxValues.Count / multiplier;// (int)(range.MaxX) - (int)(range.MinX);
+				filteredSampleCount = maxValues.Count / multiplier;// (int)(range.MaxX) - (int)(range.MinX);
 
 				// Create the filtered stat array
 				int offset = multiplier / 2;
@@ -2067,7 +2118,7 @@ namespace CSVTools
 					{
 						continue;
 					}
-					for (int i = 0; i < filteredStatCount; i++)
+					for (int i = 0; i < filteredSampleCount; i++)
 					{
 						int srcStartIndex = Math.Max(i * multiplier - offset, 0);
 						int srcEndIndex = Math.Min(i * multiplier + offset + 1, maxValues.Count);
@@ -2098,6 +2149,18 @@ namespace CSVTools
 				}
 			}
 
+			float oneOverMultiplier = 1.0f / (float)multiplier;
+
+			// Generate the event list
+			List<CsvEvent> allEvents = new List<CsvEvent>();
+			for (int i = 0; i < csvStats.Count; i++)
+			{
+				foreach (CsvEvent ev in csvStats[i].Events)
+				{
+					allEvents.Add(ev);
+				}
+			}
+
 
 			// Create a hidden panel for storing all the stat group elements
 			svg.WriteLine("<g id='interactivePanelInnerHidden<UNIQUE>' visibility='collapse'>");
@@ -2117,20 +2180,69 @@ namespace CSVTools
 				svg.WriteLine("</g>");
 			}
 			svg.WriteLine("</g>"); // interactivePanelInnerHidden
+
 			svg.WriteLine("<g id='interactivePanel<UNIQUE>' visibility='hidden'>");
 			DrawVerticalLine(svg, 0, new Colour(255, 255, 255, 0.4f), rect, range, 1.0f, true, true, "", true);
 			svg.WriteLine("<g id='interactivePanelInnerWithRect<UNIQUE>'>");
 			svg.WriteLine("<rect x='0' y='0' width='100' height='100' fill='rgba(0,0,0,0.3)' blend='1' id='interactivePanelRect<UNIQUE>' rx='5' ry='5'/>");
-
-
 			svg.WriteLine("<g id='interactivePanelInner<UNIQUE>'>");
-
 			svg.WriteLine("</g>"); // interactivePanelInner
 			svg.WriteLine("</g>"); // interactivePanelInnerWithRect
+
+			svg.WriteLine("<g id='interactiveEventPanelInnerWithRect<UNIQUE>'>");
+			svg.WriteLine("<rect x='0' y='0' width='100' height='100' fill='rgba(0,0,0,0.3)' blend='1' id='interactiveEventPanelRect<UNIQUE>' rx='5' ry='5'/>");
+			svg.WriteLine("<g id='interactiveEventPanelInner<UNIQUE>'>");
+
+			float textOffsetY = rect.y - 20;
+			DrawText(svg, "Events", 0, 9+ textOffsetY, 11, zeroRect, new Colour(255, 255, 255, 1.0f), "start", "Helvetica", "", true);
+			svg.WriteLine("</g>"); // interactiveEventPanelInner
+			svg.WriteLine("</g>"); // interactiveEventPanelInnerWithRect
 			svg.WriteLine("</g>"); // interactivepanel
 
+
 			svg.WriteLine("<script type='application/ecmascript'> <![CDATA[");
-			float oneOverMultiplier = 1.0f / (float)multiplier;
+
+			// Write the event list
+			List<string> elementStrings = new List<string>();
+			for (int i = 0; i < allEvents.Count; i++)
+			{
+				CsvEvent ev = allEvents[i];
+				elementStrings.Add("{ text: '" + ev.Name + "', frame: " + ev.Frame.ToString() + "}");
+			}
+			svg.WriteLine("var allEvents = [" + String.Join(",", elementStrings) + "]");
+
+			// Write the event index list for each sample frame
+			List<int>[] filteredEventIndexLists = new List<int>[filteredSampleCount];
+			for (int i = 0; i < filteredEventIndexLists.Length; i++)
+				filteredEventIndexLists[i] = new List<int>();
+
+			int eventPaddingPixels=6;
+			for (int i = 0; i < allEvents.Count; i++)
+			{
+				int frame = allEvents[i].Frame;
+				int filteredFrameIndex = (int)Math.Round((float)frame * oneOverMultiplier);
+
+				for (int j= filteredFrameIndex - eventPaddingPixels; j< filteredFrameIndex + eventPaddingPixels; j++)
+				{
+					if (j >= 0 && j < filteredEventIndexLists.Length)
+					{
+						filteredEventIndexLists[j].Add(i);
+					}
+				}
+			}
+
+			elementStrings = new List<string>();
+			for (int i = 0; i < filteredEventIndexLists.Length; i++)
+			{
+				List<int> eventIndexList = filteredEventIndexLists[i];
+				List<string> innerElementStrings = new List<string>();
+				foreach (int index in eventIndexList)
+				{
+					innerElementStrings.Add(index.ToString());
+				}
+				elementStrings.Add("["+String.Join(",", innerElementStrings)+"]");
+			}
+			svg.WriteLine("var eventIndexLists = [" + String.Join(",", elementStrings) + "]");
 
 			// Write the data array for each stat
 			foreach (InteractiveStatInfo statInfo in interactiveStats)
@@ -2288,13 +2400,14 @@ namespace CSVTools
 			svg.WriteLine("    var xOffset = 0;");
 			svg.WriteLine("    var lineX = ToSvgX<UNIQUE>(frameNum);");
 			svg.WriteLine("    var textX = lineX + xOffset;");
-			svg.WriteLine("    var textY = " + rect.y + " - 20");
+			svg.WriteLine("    var textY = " + textOffsetY);
 
 			svg.WriteLine("    var interactivePanelInner = document.getElementById('interactivePanelInner<UNIQUE>');");
 			svg.WriteLine("    legendPanel.setAttribute('visibility','hidden')");
 			svg.WriteLine("    interactivePanel.setAttribute('visibility','visible')");
 			svg.WriteLine("    interactivePanel.setAttribute('transform','translate('+textX+',0)')");
 			svg.WriteLine("    var dataIndex = Math.round( frameNum * " + oneOverMultiplier + " );");
+			svg.WriteLine("    dataIndex = Math.min( dataIndex, "+ (filteredSampleCount-1) +" ); ");
 
 			// Fill out the sample data array
 			svg.WriteLine("    var samples = [");
@@ -2346,6 +2459,41 @@ namespace CSVTools
 			svg.WriteLine("        }");
 			svg.WriteLine("    }");
 
+			// Display events
+			svg.WriteLine("    var eventPanelInner = document.getElementById('interactiveEventPanelInner<UNIQUE>');");
+			svg.WriteLine("    while (eventPanelInner.childNodes.length>2)");
+			svg.WriteLine("    {");
+			svg.WriteLine("        eventPanelInner.removeChild(eventPanelInner.lastChild);");
+			svg.WriteLine("    }");
+
+			svg.WriteLine("    var eventTextY = " + textOffsetY+"+15;");
+			svg.WriteLine("    var eventIndexList = eventIndexLists[dataIndex];");
+			svg.WriteLine("    for ( i=0;i<eventIndexList.length;i++ )");
+			svg.WriteLine("    {");
+			svg.WriteLine("        var index = eventIndexList[i];");
+			svg.WriteLine("        var event = allEvents[index];");
+			svg.WriteLine("        if ( eventTextY<=" + rect.height.ToString() + ")");
+			svg.WriteLine("        {");
+			svg.WriteLine("            var frameOffset = event.frame - frameNum;");
+			svg.WriteLine("            var frameOffsetStr = (frameOffset>0 ? '+' : '') +frameOffset.toString()");
+			svg.WriteLine("            var textElement = document.createElementNS('http://www.w3.org/2000/svg','text');");
+			svg.WriteLine("            textElement.setAttribute('transform','translate(0,'+eventTextY+')');");
+			svg.WriteLine("            textElement.innerHTML = event.text + ' (' + frameOffsetStr +')';");
+			svg.WriteLine("            textElement.setAttribute('fill','rgb(224, 224, 224)');");
+			svg.WriteLine("            textElement.setAttribute('filter','url(#dropShadowFilter)');");
+			svg.WriteLine("            textElement.setAttribute('stroke','CSVStats.Colour');");
+			svg.WriteLine("            textElement.setAttribute('font-size','10');");
+			svg.WriteLine("            textElement.setAttribute('font-family','Helvetica');");
+			svg.WriteLine("            textElement.setAttribute('style','text-anchor: start');");
+			svg.WriteLine("            textElement.setAttribute('x','0');");
+			svg.WriteLine("            textElement.setAttribute('y','9');");
+			svg.WriteLine("            textElement.setAttribute('visibility','inherit');");
+			svg.WriteLine("            eventPanelInner.appendChild(textElement);");
+			svg.WriteLine("            eventTextY += 12;");
+			svg.WriteLine("        }");
+			svg.WriteLine("    }");
+
+			// Set the panel rect dimensions based on the content
 			svg.WriteLine("    var panelRect = document.getElementById('interactivePanelRect<UNIQUE>');");
 			svg.WriteLine("    var bbox = interactivePanelInner.getBBox();");
 			svg.WriteLine("    panelRect.setAttribute('width',bbox.width+16);");
@@ -2365,7 +2513,32 @@ namespace CSVTools
 			svg.WriteLine("    }   ");
 			svg.WriteLine("    var interactivePanelInnerWithRect = document.getElementById('interactivePanelInnerWithRect<UNIQUE>');");
 			svg.WriteLine("    interactivePanelInnerWithRect.setAttribute('transform','translate('+panelOffset+',0)');");
+
+			// Move the event panel and set its visibility
+			svg.WriteLine("    var eventBbox = eventPanelInner.getBBox();");
+			svg.WriteLine("    if ( textX + eventBbox.width + bbox.width + bbox.x > " + maxX + " ) ");
+			svg.WriteLine("			panelOffset-=eventBbox.width + bbox.width+40;");
+
+			svg.WriteLine("    var eventPanel = document.getElementById('interactiveEventPanelInnerWithRect<UNIQUE>');");
+			svg.WriteLine("    if (eventIndexList.length>0)");
+			svg.WriteLine("    {");
+			svg.WriteLine("      eventPanel.setAttribute('visibility','inherit');");
+			svg.WriteLine("      eventPanel.setAttribute('transform','translate('+(bbox.width+bbox.x+20+panelOffset)+',0)');");
+			svg.WriteLine("    }");
+			svg.WriteLine("    else");
+			svg.WriteLine("    {");
+			svg.WriteLine("      eventPanel.setAttribute('visibility','hidden');");
+			svg.WriteLine("    }");
+
+			// Set the event panel rect size
+			svg.WriteLine("    var eventPanelRect = document.getElementById('interactiveEventPanelRect<UNIQUE>');");
+			svg.WriteLine("    eventPanelRect.setAttribute('width',eventBbox.width+16);");
+			svg.WriteLine("    eventPanelRect.setAttribute('height',eventTextY+8);");
+			svg.WriteLine("    eventPanelRect.setAttribute('x',eventBbox.x-8);");
+			svg.WriteLine("    eventPanelRect.setAttribute('y',eventBbox.y-8);");
 			svg.WriteLine("  }");
+
+			// If we moused outside the box then hide the interactive panel
 			svg.WriteLine("  else");
 			svg.WriteLine("  {");
 			svg.WriteLine("    interactivePanel.setAttribute('visibility','hidden')");
@@ -2398,21 +2571,27 @@ namespace CSVTools
 		}
 
 		bool IsEventShown(string eventString, GraphParams graphParams)
+		{
+			return GetEventPriority(eventString, graphParams) >= 0;
+		}
+
+		int GetEventPriority(string eventString, GraphParams graphParams)
 		{
 			if (graphParams.showEventNames == null)
 			{
-				return false;
+				return -1;
 			}
 			if (graphParams.showEventNames.Count == 0)
 			{
-				return false;
+				return -1;
 			}
 			if (eventString.Length == 0)
 			{
-				return false;
+				return -1;
 			}
 
 			eventString = eventString.ToLower();
+			int priority = 0;
 			foreach (string showEventName in graphParams.showEventNames)
 			{
 				string showEventNameLower = showEventName.ToLower();
@@ -2422,15 +2601,16 @@ namespace CSVTools
 					string prefix = showEventNameLower.Substring(0, index);
 					if (eventString.StartsWith(prefix))
 					{
-						return true;
+						return priority;
 					}
 				}
 				else if (eventString == showEventNameLower)
 				{
-					return true;
+					return priority;
 				}
+				priority++;
 			}
-			return false;
+			return -1;
 		}
 
 		bool IsStatIgnored(string statName, GraphParams graphParams)
diff --git a/Engine/Source/Programs/CSVTools/PerfReportTool/PerfReportTool.cs b/Engine/Source/Programs/CSVTools/PerfReportTool/PerfReportTool.cs
--- a/Engine/Source/Programs/CSVTools/PerfReportTool/PerfReportTool.cs
+++ b/Engine/Source/Programs/CSVTools/PerfReportTool/PerfReportTool.cs
@@ -1951,7 +1951,7 @@ namespace PerfReportTool
 			if (multipleCSVs)
 			{
 				graphParams.showEventNames.Add("CSV:*");
-				graphParams.showEventNameText = false;
+				graphParams.showEventNameTextMode = ShowEventTextMode.Hide;
 			}
 			else
 			{
