CsvToSvg Fix issue when minX/MaxX is specified and filteredSampleCount is incorrectly calculated

diff --git a/Engine/Source/Programs/CSVTools/CsvToSvgLib/CsvToSvgLib.cs b/Engine/Source/Programs/CSVTools/CsvToSvgLib/CsvToSvgLib.cs
--- a/Engine/Source/Programs/CSVTools/CsvToSvgLib/CsvToSvgLib.cs
+++ b/Engine/Source/Programs/CSVTools/CsvToSvgLib/CsvToSvgLib.cs
@@ -15,7 +15,7 @@ namespace CSVTools
 
 	public class CsvToSvgLibVersion
 	{
-		private static string VersionString = "3.62";
+		private static string VersionString = "3.63";
 
 		public static string Get() { return VersionString; }
 	};
@@ -2072,23 +2072,28 @@ namespace CSVTools
 			// TODO: strip out data outside the range (needs an offset as well as a multiplier)
 			int multiplier = 1;
 			float numStatsPerPixel = (float)(range.MaxX - range.MinX) / rect.width;
-			int filteredSampleCount = (int)(range.MaxX) - (int)(range.MinX);
 
-			if (numStatsPerPixel > 1)
+			// Compute the unfiltered sample count
+			int rawSampleCount = 0;
+			foreach (InteractiveStatInfo statInfo in interactiveStats)
 			{
-				multiplier = (int)numStatsPerPixel;
-
-				int maxCount = 0;
-				foreach (InteractiveStatInfo statInfo in interactiveStats)
+				if (statInfo.originalStatSamples != null)
 				{
-					if (statInfo.originalStatSamples != null)
-					{
-						maxCount = Math.Max(statInfo.originalStatSamples.samples.Count, maxCount);
-					}
+					rawSampleCount = Math.Max(statInfo.originalStatSamples.samples.Count, rawSampleCount);
 				}
+			}
+			int filteredSampleCount = rawSampleCount;
+
+			// TODO: truncate the stats/events using minX/maxX to save memory, so filteredSampleCount is (int)(range.MaxX) - (int)(range.MinX) 
+            // if not downsampling. MinX will require a bias as well as a multiplier in the javascript
+
+			// Downsample the stats if needed
+			if (numStatsPerPixel > 1)
+			{
+				multiplier = (int)numStatsPerPixel;
 
 				// Compute max value for each frame
-				List<float> maxValues = new List<float>(maxCount);
+				List<float> maxValues = new List<float>(rawSampleCount);
 				foreach (InteractiveStatInfo statInfo in interactiveStats)
 				{
 					if (statInfo.originalStatSamples != null)
