CsvToSvg 3.64: Fix javascript variable clash with event display

diff --git a/Engine/Source/Programs/CSVTools/CsvToSvgLib/CsvToSvgLib.cs b/Engine/Source/Programs/CSVTools/CsvToSvgLib/CsvToSvgLib.cs
--- a/Engine/Source/Programs/CSVTools/CsvToSvgLib/CsvToSvgLib.cs
+++ b/Engine/Source/Programs/CSVTools/CsvToSvgLib/CsvToSvgLib.cs
@@ -15,7 +15,7 @@ namespace CSVTools
 
 	public class CsvToSvgLibVersion
 	{
-		private static string VersionString = "3.63";
+		private static string VersionString = "3.64";
 
 		public static string Get() { return VersionString; }
 	};
@@ -2214,7 +2214,7 @@ namespace CSVTools
 				CsvEvent ev = allEvents[i];
 				elementStrings.Add("{ text: '" + ev.Name + "', frame: " + ev.Frame.ToString() + "}");
 			}
-			svg.WriteLine("var allEvents = [" + String.Join(",", elementStrings) + "]");
+			svg.WriteLine("var allEvents<UNIQUE> = [" + String.Join(",", elementStrings) + "]");
 
 			// Write the event index list for each sample frame
 			List<int>[] filteredEventIndexLists = new List<int>[filteredSampleCount];
@@ -2247,7 +2247,7 @@ namespace CSVTools
 				}
 				elementStrings.Add("["+String.Join(",", innerElementStrings)+"]");
 			}
-			svg.WriteLine("var eventIndexLists = [" + String.Join(",", elementStrings) + "]");
+			svg.WriteLine("var eventIndexLists<UNIQUE> = [" + String.Join(",", elementStrings) + "]");
 
 			// Write the data array for each stat
 			foreach (InteractiveStatInfo statInfo in interactiveStats)
@@ -2472,11 +2472,11 @@ namespace CSVTools
 			svg.WriteLine("    }");
 
 			svg.WriteLine("    var eventTextY = " + textOffsetY+"+15;");
-			svg.WriteLine("    var eventIndexList = eventIndexLists[dataIndex];");
+			svg.WriteLine("    var eventIndexList = eventIndexLists<UNIQUE>[dataIndex];");
 			svg.WriteLine("    for ( i=0;i<eventIndexList.length;i++ )");
 			svg.WriteLine("    {");
 			svg.WriteLine("        var index = eventIndexList[i];");
-			svg.WriteLine("        var event = allEvents[index];");
+			svg.WriteLine("        var event = allEvents<UNIQUE>[index];");
 			svg.WriteLine("        if ( eventTextY<=" + rect.height.ToString() + ")");
 			svg.WriteLine("        {");
 			svg.WriteLine("            var frameOffset = event.frame - frameNum;");
