PRT - Handle missing graphs gracefully in bulk mode and log an error instead of exiting.

This is a common mistake when editing the XML so we don't want it to prevent generating future reports.

#rnx

diff --git a/Engine/Source/Programs/CSVTools/PerfReportTool/PerfReportTool.cs b/Engine/Source/Programs/CSVTools/PerfReportTool/PerfReportTool.cs
--- a/Engine/Source/Programs/CSVTools/PerfReportTool/PerfReportTool.cs
+++ b/Engine/Source/Programs/CSVTools/PerfReportTool/PerfReportTool.cs
@@ -23,7 +23,7 @@ namespace PerfReportTool
     class Version
     {
 		// Format: Major.Minor.Bugfix
-        private static string VersionString = "4.237.4";
+        private static string VersionString = "4.237.5";
 
         public static string Get() { return VersionString; }
     };
diff --git a/Engine/Source/Programs/CSVTools/PerfReportTool/ReportXml.cs b/Engine/Source/Programs/CSVTools/PerfReportTool/ReportXml.cs
--- a/Engine/Source/Programs/CSVTools/PerfReportTool/ReportXml.cs
+++ b/Engine/Source/Programs/CSVTools/PerfReportTool/ReportXml.cs
@@ -428,6 +428,7 @@ namespace PerfReportTool
 			}
 
 			// Load the graphs
+			List<ReportGraph> invalidGraphs = new List<ReportGraph>();
 			foreach (ReportGraph graph in reportTypeInfo.graphs)
 			{
 				if (graph.isInline)
@@ -437,7 +438,17 @@ namespace PerfReportTool
 						GraphSettings parentSettings = null;
 						if (!graphs.TryGetValue(graph.parent.ToLower(), out parentSettings))
 						{
-							throw new Exception("Parent graph with title \"" + graph.parent + "\" was not found in graphs XML");
+							if (bBulkMode)
+							{
+								Console.Error.WriteLine("Parent graph with title \"" + graph.parent + "\" was not found in graphs XML. Skipping.");
+								invalidGraphs.Add(graph);
+								continue;
+							}
+							else
+							{
+								// Fatal in non-bulk mode
+								throw new Exception("Parent graph with title \"" + graph.parent + "\" was not found in graphs XML");
+							}
 						}
 						graph.settings.InheritFrom(parentSettings);
 					}
@@ -446,11 +457,23 @@ namespace PerfReportTool
 				{
 					if (!graphs.TryGetValue(graph.title.ToLower(), out graph.settings))
 					{
-						throw new Exception("Graph with title \"" + graph.title + "\" was not found in graphs XML");
+						if (bBulkMode)
+						{
+							Console.Error.WriteLine("Graph with title \"" + graph.title + "\" was not found in graphs XML. Skipping");
+							invalidGraphs.Add(graph);
+						}
+						else
+						{
+							// Fatal in non-bulk mode
+							throw new Exception("Graph with title \"" + graph.title + "\" was not found in graphs XML");
+						}
 					}
 				}
 			}
 
+			// Strip any invalid graphs
+			reportTypeInfo.graphs.RemoveAll((graph) => invalidGraphs.Contains(graph));
+
 			foreach (Summary summary in reportTypeInfo.summaries)
 			{
 				summary.PostInit(reportTypeInfo, csvFile.dummyCsvStats);
