PRT - Add sort option to bucket summary stats. Also sort event summaries by count.

+ Add the option to sort bucket summary tables by a particular bucket index (asc/desc) or by the stat name.
+ Add the option to specificy the bucket precision in the header column.
+ Sort event summaries by count by default so most frequent events are at the top.

Binary changes will be included before submission.

#rnx

diff --git a/Engine/Source/Programs/CSVTools/PerfReportTool/PerfReportTool.cs b/Engine/Source/Programs/CSVTools/PerfReportTool/PerfReportTool.cs
--- a/Engine/Source/Programs/CSVTools/PerfReportTool/PerfReportTool.cs
+++ b/Engine/Source/Programs/CSVTools/PerfReportTool/PerfReportTool.cs
@@ -23,7 +23,7 @@ namespace PerfReportTool
     class Version
     {
 		// Format: Major.Minor.Bugfix
-        private static string VersionString = "4.237.5";
+        private static string VersionString = "4.237.6";
 
         public static string Get() { return VersionString; }
     };
diff --git a/Engine/Source/Programs/CSVTools/PerfReportTool/Summaries/BucketSummary.cs b/Engine/Source/Programs/CSVTools/PerfReportTool/Summaries/BucketSummary.cs
--- a/Engine/Source/Programs/CSVTools/PerfReportTool/Summaries/BucketSummary.cs
+++ b/Engine/Source/Programs/CSVTools/PerfReportTool/Summaries/BucketSummary.cs
@@ -7,11 +7,25 @@ using System.Linq;
 using System.Xml.Linq;
 using PerfReportTool;
 using CSVStats;
+using System.Collections;
 
 namespace PerfSummaries
 {
 	class BucketSummary : Summary
 	{
+		private enum ESortMode
+		{
+			None,
+			Alpha,
+			BucketIndex
+		}
+
+		private enum ESortDirection
+		{ 
+			Ascending,
+			Descending,
+		}
+
 		public BucketSummary(XElement element, XmlVariableMappings vars, string baseXmlDirectory)
 		{
 			ReadStatsFromXML(element, vars);
@@ -19,8 +33,10 @@ namespace PerfSummaries
 
 			XElement BucketElement = element.Element("buckets");
 			bool ReportOutOfRangeDefault = BucketElement.GetSafeAttribute<bool>(vars, "reportOutOfRange", false);
+			// eg. <buckets reportAboveRange="1" reportBelowRange="1" precision="0.0">
 			ReportBelowRange = BucketElement.GetSafeAttribute<bool>(vars, "reportBelowRange", ReportOutOfRangeDefault);
 			ReportAboveRange = BucketElement.GetSafeAttribute<bool>(vars, "reportAboveRange", ReportOutOfRangeDefault);
+			ColumnHeaderPrecision = BucketElement.GetSafeAttribute<string>(vars, "precision", "0");
 
 			string[] XmlBuckets = BucketElement.Value.Split(',');
 			Buckets = new double[XmlBuckets.Length];
@@ -37,6 +53,26 @@ namespace PerfSummaries
 				LowEndColor = new Colour(ColorElement.GetSafeAttribute<string>(vars, "lowEndColor"));
 				HighEndColor = new Colour(ColorElement.GetSafeAttribute<string>(vars, "highEndColor"));
 			}
+
+			// eg. <sort mode="bucketIndex" bucketIndex="0" sortDirection="descending"></sort>
+			XElement SortElement = element.Element("sort");
+			if (SortElement != null)
+			{
+				string SortModeStr = SortElement.GetSafeAttribute<string>(vars, "mode", null);
+				SortModeStr = SortModeStr != null ? SortModeStr.ToLower() : null;
+				if (SortModeStr == "alpha")
+				{
+					SortMode = ESortMode.Alpha;
+				}
+				else if (SortModeStr == "bucketindex")
+				{
+					SortMode = ESortMode.BucketIndex;
+					SortByBucketIndex = SortElement.GetRequiredAttribute<int>(vars, "bucketIndex");
+				}
+
+				string SortDirectionStr = SortElement.GetSafeAttribute<string>(vars, "sortDirection", "ascending").ToLower();
+				SortDirection = SortDirectionStr == "ascending" ? ESortDirection.Ascending : ESortDirection.Descending;
+			}
 		}
 		public BucketSummary() { }
 
@@ -64,7 +100,7 @@ namespace PerfSummaries
 
 			if (ReportBelowRange)
 			{
-				htmlSection.WriteLine("  <th> <" + Buckets[0].ToString("0") + "</b></td>");
+				htmlSection.WriteLine("  <th> <" + Buckets[0].ToString(ColumnHeaderPrecision) + "</b></td>");
 			}
 
 			for (int i = 1; i < Buckets.Length; ++i)
@@ -72,19 +108,46 @@ namespace PerfSummaries
 				double Begin = Buckets[i - 1];
 				double End = Buckets[i];
 
-				htmlSection.WriteLine("  <th> [" + Begin.ToString("0") + ", " + End.ToString("0") + ")"+ "</b></td>");
+				htmlSection.WriteLine("  <th> [" + Begin.ToString(ColumnHeaderPrecision) + ", " + End.ToString(ColumnHeaderPrecision) + ")"+ "</b></td>");
 			}
 
 			if (ReportAboveRange)
 			{
-				htmlSection.WriteLine("  <th> >=" + Buckets.Last().ToString("0") + "</b></td>");
+				htmlSection.WriteLine("  <th> >=" + Buckets.Last().ToString(ColumnHeaderPrecision) + "</b></td>");
 			}
 
 			htmlSection.WriteLine("  </tr>");
 
-			foreach (string unitStat in stats)
+			// Create a copy of the stats if we need to sort them.
+			List<string> sortedStats = SortMode == ESortMode.None ? stats : new List<string>(stats);
+			if (SortMode == ESortMode.Alpha)
+			{
+				sortedStats.Sort((string a, string b) =>
+				{
+					return SortDirection == ESortDirection.Ascending ? a.CompareTo(b) : b.CompareTo(a);
+				});
+			}
+			else if (SortMode == ESortMode.BucketIndex)
+			{
+				if (SortByBucketIndex >= 0 && SortByBucketIndex < Buckets.Length)
+				{
+					float BucketValue = (float)Buckets[SortByBucketIndex];
+
+					sortedStats.Sort((string a, string b) =>
+					{
+						StatSamples StatsA = csvStats.GetStat(a.ToLower());
+						StatSamples StatsB = csvStats.GetStat(b.ToLower());
+						int CountA = StatsA.GetCountOfFramesAtOrOverBudget(BucketValue);
+						int CountB = StatsB.GetCountOfFramesAtOrOverBudget(BucketValue);
+						return SortDirection == ESortDirection.Ascending ? CountA.CompareTo(CountB) : CountB.CompareTo(CountA);
+					});
+				}
+			}
+
+			foreach (string unitStat in sortedStats)
 			{
 				string StatToCheck = unitStat.Split('(')[0];
+
 				StatSamples Stats = csvStats.GetStat(StatToCheck.ToLower());
 				if (Stats == null)
 				{
@@ -148,6 +211,11 @@ namespace PerfSummaries
 		public bool ReportAboveRange;
 		public Colour LowEndColor;
 		public Colour HighEndColor;
+		private string ColumnHeaderPrecision = "0";
+		// Sorting
+		ESortMode SortMode = ESortMode.None;
+		ESortDirection SortDirection;
+		private int SortByBucketIndex = -1;
 	};
 
 }
\ No newline at end of file
diff --git a/Engine/Source/Programs/CSVTools/PerfReportTool/Summaries/EventSummary.cs b/Engine/Source/Programs/CSVTools/PerfReportTool/Summaries/EventSummary.cs
--- a/Engine/Source/Programs/CSVTools/PerfReportTool/Summaries/EventSummary.cs
+++ b/Engine/Source/Programs/CSVTools/PerfReportTool/Summaries/EventSummary.cs
@@ -54,6 +54,11 @@ namespace PerfSummaries
 				}
 			}
 
+			// Sort descending
+			eventCountsDict = eventCountsDict
+				.OrderBy(pair  => -pair.Value)
+				.ToDictionary(x => x.Key, x => x.Value);
+
 			// Output HTML
 			if (bWriteHtml && eventCountsDict.Count > 0)
 			{
