Add ability for PRT to dump variables to json: -dumpVariablesToJson and -dumpAllVariablesToJson
Adds -noReports to skip report writing
Rebuild relevent binary and pdb

#rnx
#tests local PRT testing

diff --git a/Engine/Source/Programs/CSVTools/PerfReportTool/PerfReportTool.cs b/Engine/Source/Programs/CSVTools/PerfReportTool/PerfReportTool.cs
--- a/Engine/Source/Programs/CSVTools/PerfReportTool/PerfReportTool.cs
+++ b/Engine/Source/Programs/CSVTools/PerfReportTool/PerfReportTool.cs
@@ -23,7 +23,7 @@ namespace PerfReportTool
     class Version
     {
 		// Format: Major.Minor.Bugfix
-        private static string VersionString = "4.237.6";
+        private static string VersionString = "4.238.0";
 
         public static string Get() { return VersionString; }
     };
@@ -101,6 +101,7 @@ namespace PerfReportTool
 			"  -noSmooth : disable smoothing on all graphs\n" +
 			"  -listSummaryTables: lists available summary tables from the current report XML\n" +
 			"  -dumpVariables | -dumpVariablesAll : dumps variables to the log for each CSV (all includes metadata)\n" +
+			"  -dumpVariablesToJson | -dumpAllVariablesToJson <path> : path (usually a json filename) to write variables (like budgets) to json (all includes metadata)\n" +
 			"\n" +
 			"Performance args:\n" +
 			"  -perfLog : output performance logging information\n" +
@@ -128,6 +129,7 @@ namespace PerfReportTool
 			"  -customTable <comma separated fields>\n" +
 			"  -customTableSort <comma separated field row sort order> (use with -customTable)\n" +
 			"  -noDetailedReports : skips individual report generation\n" +
+			"  -noReports : skips generating reports\n" +
 			"  -collateTable : writes a collated table in addition to the main one, merging by row sort\n" +
 			"  -collateTableOnly : as -collateTable, but doesn't write the standard summary table.\n" +
 			"  -emailTable : writes a condensed email-friendly table (see the 'condensed' summary table)\n" +
@@ -473,6 +475,7 @@ namespace PerfReportTool
 			}
 
 			bool writeDetailedReports = !GetBoolArg("noDetailedReports");
+			bool writeReports = !GetBoolArg("noReports");
 			bool bReadAllStats = GetBoolArg("readAllStats");
 
 			bool bSummaryTableCacheReadonly = GetBoolArg("summaryTableCacheReadOnly");
@@ -580,13 +583,32 @@ namespace PerfReportTool
 							}
 							else
 							{
-								GenerateReport(cachedCsvFile, outputDir, bBulkMode, rowData, bBatchedGraphs, writeDetailedReports, true, cachedCsvFile.reportTypeInfo, csvDir);
-								perfLog.LogTiming("  GenerateReport");
+								if (writeReports)
+								{
+									GenerateReport(cachedCsvFile, outputDir, bBulkMode, rowData, bBatchedGraphs, writeDetailedReports, true, cachedCsvFile.reportTypeInfo, csvDir);
+									perfLog.LogTiming("  GenerateReport");
+								}
 
-								if ( ( GetBoolArg("dumpVariables") || GetBoolArg("dumpVariablesAll") ) && cachedCsvFile.xmlVariableMappings != null)
+								if (cachedCsvFile.xmlVariableMappings != null)
 								{
-									Console.WriteLine("\nDumping variables for " + cachedCsvFile.filename + "\n");
-									cachedCsvFile.xmlVariableMappings.DumpToLog(GetBoolArg("dumpVariablesAll"));
+									if ((GetBoolArg("dumpVariables") || GetBoolArg("dumpVariablesAll")))
+									{
+										Console.WriteLine("\nDumping variables for " + cachedCsvFile.filename + "\n");
+										cachedCsvFile.xmlVariableMappings.DumpToLog(GetBoolArg("dumpVariablesAll"));
+									}
+
+									string variablesJsonOutPath = GetArg("dumpVariablesToJson", null);
+									if (variablesJsonOutPath != null)
+									{
+										cachedCsvFile.xmlVariableMappings.SerializeToJson(variablesJsonOutPath, "", "meta.");
+									}
+
+									string allVariablesJsonOutPath = GetArg("dumpAllVariablesToJson", null);
+									if (allVariablesJsonOutPath != null)
+									{
+										cachedCsvFile.xmlVariableMappings.SerializeToJson(allVariablesJsonOutPath, "", "");
+									}
+
 								}
 
 								if (rowData != null && bWriteToSummaryTableCache)
diff --git a/Engine/Source/Programs/CSVTools/PerfReportTool/ReportXml.cs b/Engine/Source/Programs/CSVTools/PerfReportTool/ReportXml.cs
--- a/Engine/Source/Programs/CSVTools/PerfReportTool/ReportXml.cs
+++ b/Engine/Source/Programs/CSVTools/PerfReportTool/ReportXml.cs
@@ -8,6 +8,7 @@ using System.Xml.Linq;
 using CSVStats;
 using PerfSummaries;
 using System.Globalization;
+using System.Text.Json;
 
 namespace PerfReportTool
 {
@@ -644,6 +645,37 @@ namespace PerfReportTool
 				}
 			}
 		}
+		public void SerializeToJson(string Filename, string toMatch, string toIgnore)
+		{
+			string[] keys = vars.Keys.ToArray();
+			List<string> MatchedKeys = new List<string>();
+			foreach (string key in keys)
+			{
+				if (key.Contains(toMatch) && (toIgnore.Length==0 || !key.Contains(toIgnore)))
+				{
+					MatchedKeys.Add(key);
+				}
+			}
+
+			if (MatchedKeys.Any())
+			{
+				MatchedKeys.Sort();
+
+				Dictionary<string, string> MatchedVars = new Dictionary<string, string>();
+				foreach (string key in MatchedKeys)
+				{
+					if (vars.ContainsKey(key) && vars[key].Any())
+					{
+						MatchedVars[key] = vars[key].Replace(", ", ",").Replace(",", ", ");
+					}
+				}
+
+				JsonSerializerOptions options = new JsonSerializerOptions { WriteIndented = true };
+				FileStream createStream = File.Create(Filename);
+				JsonSerializer.Serialize(createStream, MatchedVars, options);
+				createStream.Dispose();
+			}
+		}
 
 		public string ResolveVariables(string attributeValue)
 		{
