Enhanced access to variables in PRT without CSVs
  "-metadataProxy name" looks up named metadata proxy in ReportTypes.xml csvMetadataProxies
  "-overrideMetadata name=value,name=value,name=value" supplies explicit metadata pairs from the commandline
Added query param to PRS to display variables based upon a CSV in the server database
  "/variables?id=csvid[&all=true|false]"
Updated binaries

#tests PRT test cases, local perf server

diff --git a/Engine/Source/Programs/CSVTools/PerfReportTool/PerfReportTool.cs b/Engine/Source/Programs/CSVTools/PerfReportTool/PerfReportTool.cs
--- a/Engine/Source/Programs/CSVTools/PerfReportTool/PerfReportTool.cs
+++ b/Engine/Source/Programs/CSVTools/PerfReportTool/PerfReportTool.cs
@@ -102,6 +102,8 @@ namespace PerfReportTool
 			"  -listSummaryTables: lists available summary tables from the current report XML\n" +
 			"  -dumpVariables | -dumpVariablesAll : dumps variables to the log for each CSV (all includes metadata)\n" +
 			"  -dumpVariablesToJson | -dumpAllVariablesToJson <path> : path (usually a json filename) to write variables (like budgets) to json (all includes metadata)\n" +
+			"  -metadataProxy : when dumping variables can be used in place of a csv, proxy data can be specified in a csvMetadataProxies section in report XML\n" +
+			"  -overrideMetadata : when dumping variables use to specify/override metadata: -overrideMetadata platform=abc,targetframerate=60\n" +
 			"\n" +
 			"Performance args:\n" +
 			"  -perfLog : output performance logging information\n" +
@@ -262,6 +264,20 @@ namespace PerfReportTool
 				return;
 			}
 
+			// Load the report + graph XML data
+			// Do before commandline processing, -metadataProxyName needs info from report XML
+			reportXML = new ReportXML(
+				GetArg("graphxml", false),
+				GetArg("reportxml", false),
+				GetArg("reportxmlbasedir", false),
+				GetArg("summaryTableXml", false),
+				GetArg("summaryTableXmlSubst", false),
+				GetArg("summaryTableXmlAppend", false),
+				GetArg("summaryTableXmlRowSortAppend", false)
+				);
+
+			bool cleanupTempCsv = false;
+			string tempCsv = Path.GetTempPath() + "dummy.csv";
 			string csvDir = null;
 
 			bool bBulkMode = false;
@@ -280,6 +296,8 @@ namespace PerfReportTool
 				string summaryTableCacheInDir = GetArg("summaryTableCacheIn");
 				string csvListStr = GetArg("csvList");
 				string prcListStr = GetArg("prcList");
+				string metadataProxyName = GetArg("metadataProxy");
+				string explicitMetadata = GetArg("overrideMetadata");
 				if (csvDir.Length > 0)
 				{
 					bool recurse = GetBoolArg("recurse");
@@ -331,6 +349,67 @@ namespace PerfReportTool
 					bBulkMode = true;
 					bSummaryTableCacheOnlyMode = true;
 				}
+				else if (metadataProxyName.Length > 0 || explicitMetadata.Length> 0)
+				{
+					string proxyMetadata = "";
+
+					// Look up a set of meta data from the proxy name if supplied
+					if (metadataProxyName.Length > 0)
+					{
+						proxyMetadata = reportXML.GetMetadataProxyInfo(metadataProxyName);
+						if (explicitMetadata.Length > 0)
+						{
+							proxyMetadata += ",";
+						}
+						else
+						{
+							Console.WriteLine("Warning: Unable to find metadataProxy named " + metadataProxyName + " in report xml csvMetadataProxies section.");
+						}
+					}
+
+					// Append with any supplied metadata from commandline, see -overrideMetadata
+					proxyMetadata += explicitMetadata;
+
+					// Process each key=value pair, later entries will override previous ones
+					Dictionary<string,string> metaDict = new Dictionary<string,string>();
+					if (proxyMetadata.Length > 0)
+					{
+						// This allows flexibility on the commandline, can use commas or semicolons (and spaces since they are mapped to semicolons in GetArg)
+						string[] entries = proxyMetadata.Split(',',';');
+						foreach (string entry in entries)
+						{
+							string[] data = entry.Split("=");
+							if (data.Length == 2)
+							{
+								metaDict[data[0]] = data[1];
+							}
+						}
+					}
+
+					// If we have a set of proxy metadata, save it out to a dummy.csv for processing
+					if (metaDict.Count > 0)
+					{
+						string fakeValues = "FrameTime\n1";
+						string metaDataString = "";
+						string leadingComma = "";
+						foreach (var (key, value) in metaDict)
+						{
+							metaDataString += leadingComma;
+							metaDataString += "[" + key + "]," + value;
+							leadingComma = ",";
+						}
+
+						csvFilenames = new string[] { tempCsv };
+						string[] lines = { fakeValues, metaDataString };
+						File.Delete(csvFilenames[0]);
+						File.AppendAllLines(csvFilenames[0], lines);
+						cleanupTempCsv = true;
+					}
+					else
+					{
+						throw new Exception("Valid -metadataProxy and/or -overrideMetadata not found.");
+					}
+				}
 				else
 				{
 					string csvFilenamesStr = GetArg("csv");
@@ -350,17 +429,6 @@ namespace PerfReportTool
 				}
 			}
 
-			// Load the report + graph XML data
-			reportXML = new ReportXML(
-				GetArg("graphxml", false), 
-				GetArg("reportxml", false), 
-				GetArg("reportxmlbasedir", false), 
-				GetArg("summaryTableXml", false), 
-				GetArg("summaryTableXmlSubst", false), 
-				GetArg("summaryTableXmlAppend", false),
-				GetArg("summaryTableXmlRowSortAppend", false)
-				);
-
 			if (GetBoolArg("listSummaryTables"))
 			{
 				Console.WriteLine("Listing summary tables:");
@@ -662,6 +730,10 @@ namespace PerfReportTool
 				}
 			}
 
+			if (cleanupTempCsv)
+			{
+				File.Delete(tempCsv);
+			}
 
 			if (summaryTableJsonHelper != null)
 			{
diff --git a/Engine/Source/Programs/CSVTools/PerfReportTool/ReportXml.cs b/Engine/Source/Programs/CSVTools/PerfReportTool/ReportXml.cs
--- a/Engine/Source/Programs/CSVTools/PerfReportTool/ReportXml.cs
+++ b/Engine/Source/Programs/CSVTools/PerfReportTool/ReportXml.cs
@@ -269,6 +269,19 @@ namespace PerfReportTool
 				columnFormatInfoList = new SummaryTableColumnFormatInfoCollection(summaryTableColumnInfoListEl);
 			}
 
+			// Read any metadata proxies
+			metadataProxyInfo = new Dictionary<string, string>();
+			XElement metadataProxiesMappingsElement = rootElement.Element("csvMetadataProxies");
+			if (metadataProxiesMappingsElement != null)
+			{
+				foreach (XElement proxy in metadataProxiesMappingsElement.Elements("csvMetadata"))
+				{
+					string key = proxy.FirstAttribute.Value.ToString().ToLower();
+					string value = proxy.Value.ToString().ToLower();
+					metadataProxyInfo[key] = value;
+				}
+			}
+
 			// Read the derived metadata mappings
 			derivedMetadataMappings = new DerivedMetadataMappings();
 			XElement derivedMetadataMappingsElement = rootElement.Element("derivedMetadataMappings");
@@ -587,6 +600,16 @@ namespace PerfReportTool
 			return summaryTables.Keys.ToList();
 		}
 
+		public string GetMetadataProxyInfo(string key)
+		{
+			string lowerKey = key.ToLower();
+			if (metadataProxyInfo.ContainsKey(lowerKey))
+			{
+				return metadataProxyInfo[lowerKey];
+			}
+			return "";
+		}
+
 		Dictionary<string, SummaryTableInfo> summaryTables;
 
 		XElement reportTypesElement;
@@ -599,6 +622,7 @@ namespace PerfReportTool
 		Dictionary<string, XElement> sharedSummaries;
 		Dictionary<string, GraphSettings> graphs;
 		Dictionary<string, string> statDisplayNameMapping;
+		Dictionary<string, string> metadataProxyInfo;
 		public SummaryTableColumnFormatInfoCollection columnFormatInfoList;
 		string baseXmlDirectory;
 
