Updates CSVTools projects and binaries to target anycpu instead of only x86_64
#tests FC diffs on before and after results of PerfReportFort test cases

diff --git a/Engine/Source/Programs/CSVTools/CSVCollate/CSVCollate.csproj b/Engine/Source/Programs/CSVTools/CSVCollate/CSVCollate.csproj
--- a/Engine/Source/Programs/CSVTools/CSVCollate/CSVCollate.csproj
+++ b/Engine/Source/Programs/CSVTools/CSVCollate/CSVCollate.csproj
@@ -15,7 +15,7 @@ Copyright Epic Games, Inc. All Rights Reserved.</Copyright>
     <PlatformTarget>x64</PlatformTarget>
   </PropertyGroup>
   <PropertyGroup Condition=" '$(Configuration)|$(Platform)' == 'Release|AnyCPU' ">
-    <PlatformTarget>x64</PlatformTarget>
+    <PlatformTarget>anycpu</PlatformTarget>
     <OutputPath>..\..\..\..\Binaries\DotNET\CsvTools\</OutputPath>
   </PropertyGroup>
   <ItemGroup>
diff --git a/Engine/Source/Programs/CSVTools/CSVFilter/CSVFilter.csproj b/Engine/Source/Programs/CSVTools/CSVFilter/CSVFilter.csproj
--- a/Engine/Source/Programs/CSVTools/CSVFilter/CSVFilter.csproj
+++ b/Engine/Source/Programs/CSVTools/CSVFilter/CSVFilter.csproj
@@ -15,7 +15,7 @@ Copyright Epic Games, Inc. All Rights Reserved.</Copyright>
     <PlatformTarget>x64</PlatformTarget>
   </PropertyGroup>
   <PropertyGroup Condition=" '$(Configuration)|$(Platform)' == 'Release|AnyCPU' ">
-    <PlatformTarget>x64</PlatformTarget>
+    <PlatformTarget>anycpu</PlatformTarget>
     <OutputPath>..\..\..\..\Binaries\DotNET\CsvTools\</OutputPath>
   </PropertyGroup>
   <ItemGroup>
diff --git a/Engine/Source/Programs/CSVTools/CSVSplit/CSVSplit.csproj b/Engine/Source/Programs/CSVTools/CSVSplit/CSVSplit.csproj
--- a/Engine/Source/Programs/CSVTools/CSVSplit/CSVSplit.csproj
+++ b/Engine/Source/Programs/CSVTools/CSVSplit/CSVSplit.csproj
@@ -15,7 +15,7 @@ Copyright Epic Games, Inc. All Rights Reserved.</Copyright>
     <PlatformTarget>x64</PlatformTarget>
   </PropertyGroup>
   <PropertyGroup Condition=" '$(Configuration)|$(Platform)' == 'Release|AnyCPU' ">
-    <PlatformTarget>x64</PlatformTarget>
+    <PlatformTarget>anycpu</PlatformTarget>
     <OutputPath>..\..\..\..\Binaries\DotNET\CsvTools\</OutputPath>
   </PropertyGroup>
   <ItemGroup>
diff --git a/Engine/Source/Programs/CSVTools/CSVToSVG/CSVToSVG.csproj b/Engine/Source/Programs/CSVTools/CSVToSVG/CSVToSVG.csproj
--- a/Engine/Source/Programs/CSVTools/CSVToSVG/CSVToSVG.csproj
+++ b/Engine/Source/Programs/CSVTools/CSVToSVG/CSVToSVG.csproj
@@ -9,7 +9,7 @@
     <PlatformTarget>x64</PlatformTarget>
   </PropertyGroup>
   <PropertyGroup Condition=" '$(Configuration)|$(Platform)' == 'Release|AnyCPU' ">
-    <PlatformTarget>x64</PlatformTarget>
+    <PlatformTarget>anycpu</PlatformTarget>
     <OutputPath>..\..\..\..\Binaries\DotNET\CsvTools\</OutputPath>
   </PropertyGroup>
   <PropertyGroup>
diff --git a/Engine/Source/Programs/CSVTools/CsvConvert/CsvConvert.csproj b/Engine/Source/Programs/CSVTools/CsvConvert/CsvConvert.csproj
--- a/Engine/Source/Programs/CSVTools/CsvConvert/CsvConvert.csproj
+++ b/Engine/Source/Programs/CSVTools/CsvConvert/CsvConvert.csproj
@@ -14,7 +14,7 @@
     <PlatformTarget>x64</PlatformTarget>
   </PropertyGroup>
   <PropertyGroup Condition=" '$(Configuration)|$(Platform)' == 'Release|AnyCPU' ">
-    <PlatformTarget>x64</PlatformTarget>
+    <PlatformTarget>anycpu</PlatformTarget>
     <OutputPath>..\..\..\..\Binaries\DotNET\CsvTools\</OutputPath>
   </PropertyGroup>
   <ItemGroup>
diff --git a/Engine/Source/Programs/CSVTools/CsvInfo/CsvInfo.csproj b/Engine/Source/Programs/CSVTools/CsvInfo/CsvInfo.csproj
--- a/Engine/Source/Programs/CSVTools/CsvInfo/CsvInfo.csproj
+++ b/Engine/Source/Programs/CSVTools/CsvInfo/CsvInfo.csproj
@@ -14,7 +14,7 @@
     <PlatformTarget>x64</PlatformTarget>
   </PropertyGroup>
   <PropertyGroup Condition=" '$(Configuration)|$(Platform)' == 'Release|AnyCPU' ">
-    <PlatformTarget>x64</PlatformTarget>
+    <PlatformTarget>anycpu</PlatformTarget>
     <OutputPath>..\..\..\..\Binaries\DotNET\CsvTools\</OutputPath>
   </PropertyGroup>
   <ItemGroup>
diff --git a/Engine/Source/Programs/CSVTools/CsvStats/CsvStats.csproj b/Engine/Source/Programs/CSVTools/CsvStats/CsvStats.csproj
--- a/Engine/Source/Programs/CSVTools/CsvStats/CsvStats.csproj
+++ b/Engine/Source/Programs/CSVTools/CsvStats/CsvStats.csproj
@@ -14,7 +14,7 @@ Copyright Epic Games, Inc. All Rights Reserved.</Copyright>
   </PropertyGroup>
   <PropertyGroup Condition=" '$(Configuration)|$(Platform)' == 'Release|AnyCPU' ">
     <OutputPath>..\..\..\..\Binaries\DotNET\CsvTools\</OutputPath>
-    <PlatformTarget>x64</PlatformTarget>
+    <PlatformTarget>anycpu</PlatformTarget>
   </PropertyGroup>
   <ItemGroup>
     <PackageReference Include="Microsoft.CSharp" Version="4.7.0" />
diff --git a/Engine/Source/Programs/CSVTools/CsvToSvgLib/CsvToSvgLib.csproj b/Engine/Source/Programs/CSVTools/CsvToSvgLib/CsvToSvgLib.csproj
--- a/Engine/Source/Programs/CSVTools/CsvToSvgLib/CsvToSvgLib.csproj
+++ b/Engine/Source/Programs/CSVTools/CsvToSvgLib/CsvToSvgLib.csproj
@@ -17,7 +17,7 @@ Copyright Epic Games, Inc. All Rights Reserved.</Copyright>
   </PropertyGroup>
   <PropertyGroup Condition=" '$(Configuration)|$(Platform)' == 'Release|AnyCPU' ">
     <OutputPath>..\..\..\..\Binaries\DotNET\CsvTools\</OutputPath>
-    <PlatformTarget>x64</PlatformTarget>
+    <PlatformTarget>anycpu</PlatformTarget>
   </PropertyGroup>
   <ItemGroup>
     <ProjectReference Include="..\CsvStats\CSVStats.csproj" />
diff --git a/Engine/Source/Programs/CSVTools/PerfReportTool/PerfreportTool.csproj b/Engine/Source/Programs/CSVTools/PerfReportTool/PerfreportTool.csproj
--- a/Engine/Source/Programs/CSVTools/PerfReportTool/PerfreportTool.csproj
+++ b/Engine/Source/Programs/CSVTools/PerfReportTool/PerfreportTool.csproj
@@ -9,7 +9,7 @@
     <PlatformTarget>x64</PlatformTarget>
   </PropertyGroup>
   <PropertyGroup Condition=" '$(Configuration)|$(Platform)' == 'Release|AnyCPU' ">
-    <PlatformTarget>x64</PlatformTarget>
+    <PlatformTarget>anycpu</PlatformTarget>
     <OutputPath>..\..\..\..\Binaries\DotNET\CsvTools\</OutputPath>
   </PropertyGroup>
   <ItemGroup>
