PerfReportTool 4.239.0 : Make -customTable -customTableSort and more flexible. These can now be combined with -summaryTable to override the filter or row sort respectively.

diff --git a/Engine/Source/Programs/CSVTools/PerfReportTool/PerfReportTool.cs b/Engine/Source/Programs/CSVTools/PerfReportTool/PerfReportTool.cs
--- a/Engine/Source/Programs/CSVTools/PerfReportTool/PerfReportTool.cs
+++ b/Engine/Source/Programs/CSVTools/PerfReportTool/PerfReportTool.cs
@@ -23,7 +23,7 @@ namespace PerfReportTool
     class Version
     {
 		// Format: Major.Minor.Bugfix
-        private static string VersionString = "4.238.0";
+        private static string VersionString = "4.239.0";
 
         public static string Get() { return VersionString; }
     };
@@ -128,8 +128,8 @@ namespace PerfReportTool
 			"Optional bulk mode args: (use with -csvdir, -summaryTableCacheIn, -csvList, -prcList)\n" +
 			"  -recurse \n" +
 			"  -searchpattern <pattern>, e.g -searchpattern csvprofile*\n" +
-			"  -customTable <comma separated fields>\n" +
-			"  -customTableSort <comma separated field row sort order> (use with -customTable)\n" +
+			"  -customTable <comma separated fields> - sets or overrides the summary table metrics filter\n" +
+			"  -customTableSort <comma separated field row sort order> - overrides the summary table row sort\n" +
 			"  -noDetailedReports : skips individual report generation\n" +
 			"  -noReports : skips generating reports\n" +
 			"  -collateTable : writes a collated table in addition to the main one, merging by row sort\n" +
@@ -214,15 +214,6 @@ namespace PerfReportTool
 			"     Use this if you want to avoid cache data being invalidated by report changes\n" +
 			"  -noCsvCacheFiles: disables usage of .csv.cache files. Cache files can be much faster if filtering on metadata\n" +
 			"";
-		/*
-		"Note on custom tables:\n" +
-		"       The -customTable and -customTableSort args allow you to generate a custom summary table\n" +
-		"       This is an alternative to using preset summary tables (see -summarytable)\n" +
-		"       Example:\n"+
-		"               -customTableSort \"deviceprofile,buildversion\" -customTable \"deviceprofile,buildversion,memoryfreeMB*\" \n"+
-		"       This outputs a table containing deviceprofile, buildversion and memoryfree stats, sorted by deviceprofile and then buildversion\n" +
-		""
-		*/
 
 		Dictionary<string, string> statDisplaynameMapping;
 		ReportXML reportXML;
@@ -555,8 +546,8 @@ namespace PerfReportTool
 			}
 
 			bool bShowHiddenStats = GetBoolArg("showHiddenStats");
-			string customSummaryTableFilter = GetArg("customTable");
-			if (customSummaryTableFilter.Length > 0)
+			string customSummaryTableFilter = GetArg("customTable", null);
+			if (customSummaryTableFilter != null)
 			{
 				bShowHiddenStats = true;
 			}
@@ -769,45 +760,28 @@ namespace PerfReportTool
 
 				bool bSpreadsheetFriendlyStrings = GetBoolArg("spreadsheetFriendly");
 				string weightByColumnName = GetArg("weightByColumn", null);
-				if (customSummaryTableFilter.Length > 0)
+
+				string customSummaryTableRowSort = GetArg("customTableSort", null);
+				string summaryTableName = GetArg("summaryTable",null);
+				SummaryTableInfo tableInfo = GetTableInfo(summaryTableName, customSummaryTableFilter, customSummaryTableRowSort);
+				if (!bCollateTableOnly)
 				{
-					string customSummaryTableRowSort = GetArg("customTableSort");
-					if (customSummaryTableRowSort.Length == 0)
-					{
-						customSummaryTableRowSort = "buildversion,deviceprofile";
-					}
-					if (!bCollateTableOnly)
-					{
-						WriteSummaryTableReport(outputDir, summaryTableFilename, summaryTable, customSummaryTableFilter.Split(',').ToList(), customSummaryTableRowSort.Split(',').ToList(), false, bCsvTable, bSpreadsheetFriendlyStrings, null, null);
-					}
-					if (bCollateTable)
-					{
-						WriteSummaryTableReport(outputDir, collatedTableFilename, summaryTable, customSummaryTableFilter.Split(',').ToList(), customSummaryTableRowSort.Split(',').ToList(), true, bCsvTable, bSpreadsheetFriendlyStrings, null, weightByColumnName);
-					}
+					WriteSummaryTableReport(outputDir, summaryTableFilename, summaryTable, tableInfo, false, bCsvTable, bSpreadsheetFriendlyStrings, null);
 				}
-				else
+				if (bCollateTable)
 				{
-					string summaryTableName = GetArg("summaryTable");
-					if (summaryTableName.Length == 0)
-					{
-						summaryTableName = "default";
-					}
-					SummaryTableInfo tableInfo = reportXML.GetSummaryTable(summaryTableName);
-					if (!bCollateTableOnly)
-					{
-						WriteSummaryTableReport(outputDir, summaryTableFilename, summaryTable, tableInfo, false, bCsvTable, bSpreadsheetFriendlyStrings, null);
-					}
-					if (bCollateTable)
-					{
-						WriteSummaryTableReport(outputDir, collatedTableFilename, summaryTable, tableInfo, true, bCsvTable, bSpreadsheetFriendlyStrings, weightByColumnName);
-					}
+					WriteSummaryTableReport(outputDir, collatedTableFilename, summaryTable, tableInfo, true, bCsvTable, bSpreadsheetFriendlyStrings, weightByColumnName);
 				}
 
 				// EmailTable is hardcoded to use the condensed type
-				string condensedSummaryTable = GetArg("condensedSummaryTable", null);
-				if (GetBoolArg("emailSummary") || GetBoolArg("emailTable") || condensedSummaryTable != null)
+				string condensedSummaryTableName = GetArg("condensedSummaryTable", null);
+				if (GetBoolArg("emailSummary") || GetBoolArg("emailTable") || condensedSummaryTableName != null)
 				{
-					SummaryTableInfo tableInfo = reportXML.GetSummaryTable(condensedSummaryTable == null ? "condensed" : condensedSummaryTable);
+					if (condensedSummaryTableName == null )
+					{
+						condensedSummaryTableName = "condensed";
+					}
+					tableInfo = GetTableInfo(condensedSummaryTableName, customSummaryTableFilter, customSummaryTableRowSort);
 					WriteSummaryTableReport(outputDir, summaryTableFilename + "_Email", summaryTable, tableInfo, true, false, bSpreadsheetFriendlyStrings, weightByColumnName);
 				}
 				perfLog.LogTiming("WriteSummaryTable");
@@ -821,6 +795,41 @@ namespace PerfReportTool
 			perfLog.LogTotalTiming();
 		}
 
+		SummaryTableInfo GetTableInfo(string summaryTableName, string customSummaryTableFilterString = null, string customSummaryTableRowSortString = null)
+		{
+			SummaryTableInfo tableInfo;
+
+			// Use the default summary table name if one isn't provided, unless we're making a custom table			
+			if (summaryTableName == null && customSummaryTableFilterString == null )
+			{
+				summaryTableName = "default";
+			}
+
+			if (summaryTableName != null)
+			{
+				tableInfo = reportXML.GetSummaryTable(summaryTableName);
+			}
+			else
+			{
+				tableInfo = new SummaryTableInfo();
+			}
+			if (customSummaryTableFilterString != null)
+			{
+				tableInfo.columnFilterList = customSummaryTableFilterString.Split(',').ToList();
+				
+				// If we're doing a custom sort but don't have a provided row sort (or preset summary table), set a default
+				if (summaryTableName == null && customSummaryTableRowSortString == null)
+				{
+					customSummaryTableRowSortString = "buildversion,deviceprofile";
+				}
+			}
+			if (customSummaryTableRowSortString != null)
+			{
+				tableInfo.rowSortList = customSummaryTableRowSortString.Split(',').ToList();
+			}
+			return tableInfo;
+		}
+
 		void WriteSummaryTableReport(string outputDir, string filenameWithoutExtension, SummaryTable table, List<string> columnFilterList, List<string> rowSortList, bool bCollated, bool bToCSV, bool bSpreadsheetFriendlyStrings, List<SummarySectionBoundaryInfo> sectionBoundaries, string weightByColumnName)
 		{
 			SummaryTableInfo tableInfo = new SummaryTableInfo();
