PerfReportTool 4.240.0 - improved diff table handling
- Add -DiffRowsAlternating param so we can control the frequency of diff rows. If enabled (default), a diff row is displayed after ever row after the second one. If disabled, we only display a diff row after every pair of rows.
- Add -ShowOnlyDiffRows param to hide non-diff rows and combine indentical text values from the source rows. This is useful for making compact tables of deltas
- Modify summaryTableXmlSubst to support one-to-many remapping. This allows more flexible modifications, such as insertion, e.g -summaryTableXmlSubst "buildversion=buildversion,deviceprofile".
NOTE: We now use a semicolon as the top level separator for remappings, since commas are used internally

CsvStats:
- Add support for GetBoolArg passing a default and support passing 1 or 0 to any bool arg
#tests Tested on local PRS server and via PRT test suite

diff --git a/Engine/Source/Programs/CSVTools/CsvStats/CommandLineTool.cs b/Engine/Source/Programs/CSVTools/CsvStats/CommandLineTool.cs
--- a/Engine/Source/Programs/CSVTools/CsvStats/CommandLineTool.cs
+++ b/Engine/Source/Programs/CSVTools/CsvStats/CommandLineTool.cs
@@ -169,9 +169,21 @@ namespace CSVStats
 			return defaultValue;
 		}
 
-		public bool GetBoolArg(string key)
+		public bool GetBoolArg(string key, bool defaultValue)
 		{
-			return CommandLineArgs.ContainsKey(key.ToLower());
+			if ( CommandLineArgs.ContainsKey(key.ToLower()) )
+			{
+				string value = CommandLineArgs[key.ToLower()];
+				if (value == "1")
+				{
+					return true;
+				}
+				if (value == "0")
+				{
+					return false;
+				}
+			}
+			return defaultValue;
 		}
 
 		public bool? GetOptionalBoolArg(string key)
@@ -265,9 +277,9 @@ namespace CSVStats
 			return commandLine.GetFloatArg(key, defaultValue);
 		}
 
-		protected bool GetBoolArg(string key)
+		protected bool GetBoolArg(string key, bool defaultValue=false)
         {
-			return commandLine.GetBoolArg(key);
+			return commandLine.GetBoolArg(key, defaultValue);
         }
 
 		protected bool? GetOptionalBoolArg(string key)
diff --git a/Engine/Source/Programs/CSVTools/PerfReportTool/PerfReportTool.cs b/Engine/Source/Programs/CSVTools/PerfReportTool/PerfReportTool.cs
--- a/Engine/Source/Programs/CSVTools/PerfReportTool/PerfReportTool.cs
+++ b/Engine/Source/Programs/CSVTools/PerfReportTool/PerfReportTool.cs
@@ -23,7 +23,7 @@ namespace PerfReportTool
     class Version
     {
 		// Format: Major.Minor.Bugfix
-        private static string VersionString = "4.239.0";
+        private static string VersionString = "4.240.0";
 
         public static string Get() { return VersionString; }
     };
@@ -169,7 +169,8 @@ namespace PerfReportTool
 			"  -maxFileAgeDays <n> : max file age in days. CSV or PRC files older than this will be ignored\n" +
 			"  -summaryTableStatThreshold <n> : stat/metric columns in the summarytable will be filtered out if all values are\n" +
 			"     less than the threshold\n" +
-			"  -summaryTableXmlSubst <find1>=<replace1>,<find2>=<replace2>... : replace summarytable XML row and filter entries\n" +
+			"  -summaryTableXmlSubst <find1>=<replace1>;<find2>=<replace2>... : replace summarytable XML row and filter entries\n" +
+			"                                                                   Note: one to many is supported (using , separators)\n" +
 			"  -summaryTableXmlAppend <list,of,stats> : append these stats to the summary table's filter list\n" +
 			"  -summaryTableXmlRowSortAppend <list,of,stats> : append these stats to the summary table's row sort list\n" +
 			"  -transposeTable : write the summary tables transposed\n" +
@@ -180,6 +181,8 @@ namespace PerfReportTool
 			"      Notes: Diff score corresponds to the value of a column's diff row; the sign is reversed if LowIsBad\n" +
 			"             Stats are sorted by prefixes first, e.g GPU/\n" +
 			"  -columnDiffDisplayThreshold <value> : if specified, hides columns with max diff value below this threshold\n"+
+			"  -diffRowsAlternating 0|1. If enabled, every other column has a diff. Otherwise diffs are only generated for pairs\n" +
+			"  -showOnlyDiffRows: If enabled, non-diff rows are hidden. Common text fields are combined.\n" +
 			"\n" +
 			"Optional Column Filters\n" +
 			"  -hideMetadataColumns : filters out metadata columns from the table (excluding those used in row sort).\n" +
@@ -920,7 +923,7 @@ namespace PerfReportTool
 				string summaryTitle = GetArg("summaryTitle", null);
 				if (GetBoolArg("addDiffRows"))
 				{
-					filteredTable.AddDiffRows(GetBoolArg("sortColumnsByDiff"), GetFloatArg("columnDiffDisplayThreshold", 0.0f));
+					filteredTable.AddDiffRows(GetBoolArg("sortColumnsByDiff"), GetFloatArg("columnDiffDisplayThreshold", 0.0f), GetBoolArg("showOnlyDiffRows"), GetBoolArg("diffRowsAlternating", true));
 				}
 
 				// Run again to add format info for any new columns that were added (eg. count).
diff --git a/Engine/Source/Programs/CSVTools/PerfReportTool/ReportXml.cs b/Engine/Source/Programs/CSVTools/PerfReportTool/ReportXml.cs
--- a/Engine/Source/Programs/CSVTools/PerfReportTool/ReportXml.cs
+++ b/Engine/Source/Programs/CSVTools/PerfReportTool/ReportXml.cs
@@ -316,17 +316,17 @@ namespace PerfReportTool
 			if (summaryTablesElement != null)
 			{
 				// Read the substitutions
-				Dictionary<string, string> substitutionsDict = null;
-				string[] substitutions = summaryTableXmlSubstStr.Split(',');
+				Dictionary<string, List<string>> substitutionsDict = null;
+				string[] substitutions = summaryTableXmlSubstStr.Split(';');
 				if (substitutions.Length>0)
 				{
-					substitutionsDict = new Dictionary<string, string>();
+					substitutionsDict = new Dictionary<string, List<string>>();
 					foreach (string substStr in substitutions)
 					{
 						string [] pair = substStr.Split('=');
 						if (pair.Length == 2)
 						{
-							substitutionsDict[pair[0]] = pair[1];
+							substitutionsDict[pair[0]] = pair[1].Split(",").ToList();
 						}
 					}
 				}
diff --git a/Engine/Source/Programs/CSVTools/PerfReportTool/SummaryTable.cs b/Engine/Source/Programs/CSVTools/PerfReportTool/SummaryTable.cs
--- a/Engine/Source/Programs/CSVTools/PerfReportTool/SummaryTable.cs
+++ b/Engine/Source/Programs/CSVTools/PerfReportTool/SummaryTable.cs
@@ -66,7 +66,7 @@ namespace PerfSummaries
 
 	class SummaryTableInfo
 	{
-		public SummaryTableInfo(XElement tableElement, Dictionary<string,string> substitutionsDict, string[] appendList, string[] rowSortAppendList, XmlVariableMappings variableMappings )
+		public SummaryTableInfo(XElement tableElement, Dictionary<string,List<string>> substitutionsDict, string[] appendList, string[] rowSortAppendList, XmlVariableMappings variableMappings )
 		{
 			string rowSortStr = tableElement.GetSafeAttribute<string>(variableMappings, "rowSort");
 			if (rowSortStr != null)
@@ -139,7 +139,7 @@ namespace PerfSummaries
 			}
 		}
 
-		private void ApplySubstitutionsToList(List<string> list, Dictionary<string, string> substitutionsDict)
+		private void ApplySubstitutionsToList(List<string> list, Dictionary<string, List<string>> substitutionsDict)
 		{
 			if (substitutionsDict == null)
 			{
@@ -147,19 +147,29 @@ namespace PerfSummaries
 			}
 			for (int i = 0; i < list.Count; i++)
 			{
-				list[i] = ApplySubstitution(list[i], substitutionsDict);
+				if (substitutionsDict.TryGetValue(list[i], out List<string> replaceList))
+				{
+					list.RemoveAt(i);
+					list.InsertRange(i, replaceList);
+					i += replaceList.Count - 1;
+				}
 			}
 		}
 
-		private string ApplySubstitution(string str, Dictionary<string, string> substitutionsDict)
+		private string ApplySubstitution(string str, Dictionary<string, List<string>> substitutionsDict)
 		{
 			if (substitutionsDict == null)
 			{
 				return str;
 			}
-			if (substitutionsDict.TryGetValue(str, out string replaceStr))
+			if (substitutionsDict.TryGetValue(str, out List<string> replaceList))
 			{
-				return replaceStr;
+				if (replaceList.Count>1 || replaceList.Count == 0)
+				{
+					// This method doesn't support one-to-many remapping, so ignore
+					return str;
+				}
+				return replaceList[0];
 			}
 			return str;
 		}
@@ -257,6 +267,13 @@ namespace PerfSummaries
 		Max
 	};
 
+	enum DiffRowFrequency
+	{
+		None,
+		Alternating,
+		AfterEachPair,
+		All
+	};
 
 	class SummaryTableColumnFormatInfo
 	{
@@ -352,7 +369,8 @@ namespace PerfSummaries
 		public bool isNumeric = false;
 		public string displayName;
 		public bool isRowWeightColumn = false;
-		public bool hasDiffRows = false;
+		public DiffRowFrequency diffRowFrequency = DiffRowFrequency.None;
+		public bool bShowOnlyDiffRows = false;
 		public bool isCountColumn = false;
 		// Column header tooltip. Displayed when hovering over the header.
 		public string tooltip = null;
@@ -430,7 +448,7 @@ namespace PerfSummaries
 			newColumn.colourThresholds.AddRange(colourThresholds);
 			newColumn.toolTips.AddRange(toolTips);
 			newColumn.colourModifiers = colourModifiers.ToDictionary(entry => entry.Key, entry => entry.Value); // Deep copy
-			newColumn.hasDiffRows = hasDiffRows;
+			newColumn.diffRowFrequency = diffRowFrequency;
 			return newColumn;
 		}
 
@@ -460,15 +478,21 @@ namespace PerfSummaries
 			return value == double.MaxValue ? 0.0 : value;
 		}
 
-		public void AddDiffRows(bool bIsFirstColumn=false)
+
+
+		public void AddDiffRows(bool bIsFirstColumn, DiffRowFrequency inDiffRowFrequency, bool bShowOnlyDiffRows)
 		{
-			if (hasDiffRows)
+			if (diffRowFrequency != DiffRowFrequency.None)
 			{
 				throw new Exception("Column already has diff rows!");
 			}
 			// Add a diff row for every row after the first one
 			int oldCount = GetCount();
-			int diffRowCount = oldCount - 1;
+
+			diffRowFrequency = inDiffRowFrequency;
+			bool bDiffRowsAlternating = diffRowFrequency == DiffRowFrequency.Alternating;
+
+			int diffRowCount = bDiffRowsAlternating ? oldCount - 1 : oldCount / 2;
 			int newCount = oldCount + diffRowCount;
 
 			// Create new lists with counts reserved
@@ -479,11 +503,16 @@ namespace PerfSummaries
 
 			bool bComputeDiff = isNumeric && !isCountColumn;
 
+			static bool NeedsDiffColumn(int originalRowIndex, bool bShowIntermediateDiffRows)
+			{
+				return bShowIntermediateDiffRows ? originalRowIndex > 0 : originalRowIndex % 2 == 1;
+			}
+
 			// Add diff rows to each of the arrays
 			for (int i = 0; i < doubleValues.Count; i++)
 			{
 				newDoubleValues.Add(doubleValues[i]);
-				if (i > 0)
+				if (NeedsDiffColumn(i, bDiffRowsAlternating))
 				{
 					if (bComputeDiff)
 					{
@@ -493,22 +522,48 @@ namespace PerfSummaries
 					}
 					else
 					{
-						newDoubleValues.Add(0.0);
+						double valueToShow = 0.0;
+						if (bShowOnlyDiffRows && isCountColumn )
+						{
+							// If we're showing only diff rows then display the count column if the values are equal
+							double thisValue = FilterInvalidValue(doubleValues[i]);
+							double prevValue = FilterInvalidValue(doubleValues[i - 1]);
+							if (thisValue == prevValue)
+							{
+								valueToShow = prevValue;
+							}
+						}
+
+						newDoubleValues.Add(valueToShow);
 					}
 				}
 			}
 			for (int i = 0; i < stringValues.Count; i++)
 			{
 				newStringValues.Add(stringValues[i]);
-				if (i > 0)
+				if (NeedsDiffColumn(i, bDiffRowsAlternating))
 				{
-					newStringValues.Add(bIsFirstColumn ? "Diff" : "");
+					if (bShowOnlyDiffRows)
+					{
+						if (stringValues[i] == stringValues[i-1])
+						{
+							newStringValues.Add(stringValues[i]); 
+						}
+						else
+						{
+							newStringValues.Add("");
+						}
+					}
+					else
+					{
+						newStringValues.Add(bIsFirstColumn ? "Diff" : "");
+					}					
 				}
 			}
 			for (int i = 0; i < toolTips.Count; i++)
 			{
 				newToolTips.Add(toolTips[i]);
-				if (i > 0)
+				if (NeedsDiffColumn(i, bDiffRowsAlternating))
 				{
 					newToolTips.Add("");
 				}
@@ -516,7 +571,7 @@ namespace PerfSummaries
 			for (int i = 0; i < colourThresholds.Count; i++)
 			{
 				newColourThresholds.Add(colourThresholds[i]);
-				if (i > 0)
+				if (NeedsDiffColumn(i, bDiffRowsAlternating))
 				{
 					newColourThresholds.Add(null);
 				}
@@ -526,33 +581,32 @@ namespace PerfSummaries
 			stringValues = newStringValues;
 			toolTips = newToolTips;
 			colourThresholds = newColourThresholds;
-			hasDiffRows = true;
 		}
 
-		public bool IsDiffRow(int rowIndex)
+		bool IsDiffRow(int rowIndex)
 		{
-			if (hasDiffRows == false || rowIndex < 2)
-			{
-				return false;
-			}
-			return ((rowIndex - 2) % 2) == 0;
+			return SummaryTable.IsDiffRow(diffRowFrequency, rowIndex);
 		}
 
 		// Computes a score (significance indicator) for a column based on its diff values. This takes into account the max value. If LowIsBad for this column then the sign is reversed
 		public double GetDiffScore()
 		{
-			if (!hasDiffRows || !isNumeric)
+			if (diffRowFrequency == DiffRowFrequency.None || !isNumeric)
 			{
 				return 0.0;
 			}
 			bool bLowIsBad = formatInfo != null && formatInfo.autoColorizeMode == AutoColorizeMode.LowIsBad;
 
 			// Find the max of all diff values for this column. If LowIsBad then we reverse the sign
+
 			double maxDiffScore = double.MinValue;
-			for (int diffRowIndex = 2; diffRowIndex < GetCount(); diffRowIndex += 2)
+			for (int rowIndex = 0; rowIndex < GetCount(); rowIndex++)
 			{
-				double diffValue = GetValue(diffRowIndex);
-				maxDiffScore = Math.Max(maxDiffScore, bLowIsBad ? -diffValue : diffValue);
+				if (IsDiffRow(rowIndex))
+				{
+					double diffValue = GetValue(rowIndex);
+					maxDiffScore = Math.Max(maxDiffScore, bLowIsBad ? -diffValue : diffValue);
+				}
 			}
 			return maxDiffScore;
 		}
@@ -560,14 +614,17 @@ namespace PerfSummaries
 		// Computes the max of the abs diff values for a column
 		public double GetMaxAbsDiff()
 		{
-			if (!hasDiffRows || !isNumeric)
+			if (diffRowFrequency == DiffRowFrequency.None || !isNumeric)
 			{
 				return 0.0;
 			}
 			double maxAbsDiff = double.MinValue;
-			for (int diffRowIndex = 2; diffRowIndex < GetCount(); diffRowIndex += 2)
+			for (int rowIndex = 0; rowIndex < GetCount(); rowIndex++)
 			{
-				maxAbsDiff = Math.Max(maxAbsDiff, Math.Abs(GetValue(diffRowIndex)));
+				if (IsDiffRow(rowIndex))
+				{
+					maxAbsDiff = Math.Max(maxAbsDiff, Math.Abs(GetValue(rowIndex)));
+				}
 			}
 			return maxAbsDiff;
 		}
@@ -715,7 +772,12 @@ namespace PerfSummaries
 		public string GetTextColor(int index)
 		{
 			const double absoluteIgnoreThreshold = 0.025;
-			if (hasDiffRows && IsDiffRow(index) && isNumeric && index < doubleValues.Count )
+        	if (bShowOnlyDiffRows && isCountColumn)
+			{
+				// If we're showing only diff rows then don't display the count column's diff row as a diff for formatting
+				return null;
+			}
+			if (IsDiffRow(index) && isNumeric && index < doubleValues.Count )
 			{
 				// For simplicity, just negate the diff value if lowIsBad
 				double diffValue = doubleValues[index];
@@ -909,7 +971,13 @@ namespace PerfSummaries
 				double val = doubleValues[index];
 
 				string prefix = "";
-				bool bIsDiffRow = hasDiffRows && IsDiffRow(index);
+				bool bIsDiffRow = IsDiffRow(index);
+	        	if (bShowOnlyDiffRows && isCountColumn)
+				{
+					// If we're showing only diff rows then don't display the count column's diff row as a diff for formatting
+					bIsDiffRow=false;
+				}
+
 				if (bIsDiffRow)
 				{
 					if ( val == 0.0 )
@@ -1528,14 +1596,23 @@ namespace PerfSummaries
 			return newTable;
 		}
 
-		public void AddDiffRows(bool bSortColumnsByDiff, double columnDiffDisplayThreshold)
+		public void AddDiffRows(bool bSortColumnsByDiff, double columnDiffDisplayThreshold, bool bInShowOnlyDiffRows, bool bDiffRowsAlternating)
 		{
+			diffRowFrequency = bDiffRowsAlternating ? DiffRowFrequency.Alternating : DiffRowFrequency.AfterEachPair;
+			bShowOnlyDiffRows = bInShowOnlyDiffRows;
 			for (int i=0; i<columns.Count; i++)
 			{
-				columns[i].AddDiffRows(i == 0);
+				columns[i].AddDiffRows(i == 0, diffRowFrequency, bInShowOnlyDiffRows);
 			}
-			rowCount += rowCount - 1;
 
+			if (bDiffRowsAlternating)
+			{
+				rowCount += rowCount - 1;
+			}
+			else
+			{
+				rowCount += rowCount / 2;
+			}
 			if ( columnDiffDisplayThreshold > 0.0 )
 			{
 				FilterColumnsByDiffThreshold(columnDiffDisplayThreshold);
@@ -2050,6 +2127,10 @@ namespace PerfSummaries
 			// Add the rows to the table
 			for (int rowIndex = 0; rowIndex < rowCount; rowIndex++)
 			{
+				if (!IsRowVisible(rowIndex))
+				{
+					continue;
+				}
 				string rowClassStr = "";
 
 				// Is this a major/minor section boundary
@@ -2335,6 +2416,42 @@ namespace PerfSummaries
 			return columnLookup.Get(name);
 		}
 
+		public bool IsRowVisible(int rowIndex)
+		{
+			return SummaryTable.IsRowVisible(diffRowFrequency, bShowOnlyDiffRows, rowIndex);
+		}
+
+		public static bool IsRowVisible(DiffRowFrequency diffRowFrequency, bool bShowOnlyDiffRows, int rowIndex)
+		{
+			if (bShowOnlyDiffRows)
+			{
+				return IsDiffRow(diffRowFrequency, rowIndex);
+			}
+			else
+			{
+				return true;
+			}
+		}
+
+		public static bool IsDiffRow(DiffRowFrequency diffRowFrequency, int rowIndex)
+		{
+			if (rowIndex < 2)
+			{
+				return false;
+			}
+			switch (diffRowFrequency)
+			{
+				case DiffRowFrequency.Alternating:
+					return ((rowIndex - 2) % 2) == 0;
+				case DiffRowFrequency.AfterEachPair:
+					return ((rowIndex - 2) % 3) == 0;
+				case DiffRowFrequency.None:
+				default:
+					return false;
+			}
+		}
+
+
 		SummaryTableColumnLookup columnLookup = new SummaryTableColumnLookup();
 		List<SummaryTableColumn> columns = new List<SummaryTableColumn>();
 		List<double> rowWeightings = null;
@@ -2342,6 +2459,8 @@ namespace PerfSummaries
 		int firstStatColumnIndex = 0;
 		bool isCollated = false;
 		bool hasMinMaxColumns = false;
+		DiffRowFrequency diffRowFrequency = DiffRowFrequency.None;
+		bool bShowOnlyDiffRows = false;
 	};
 
 	class HtmlTable
