PerfReportTool - Fix count row diff formatting when ShowOnlyDiffRows is specified

diff --git a/Engine/Source/Programs/CSVTools/PerfReportTool/SummaryTable.cs b/Engine/Source/Programs/CSVTools/PerfReportTool/SummaryTable.cs
--- a/Engine/Source/Programs/CSVTools/PerfReportTool/SummaryTable.cs
+++ b/Engine/Source/Programs/CSVTools/PerfReportTool/SummaryTable.cs
@@ -772,7 +772,7 @@ namespace PerfSummaries
 		public string GetTextColor(int index)
 		{
 			const double absoluteIgnoreThreshold = 0.025;
-        	if (bShowOnlyDiffRows && isCountColumn)
+        	if (!bShowOnlyDiffRows && isCountColumn)
 			{
 				// If we're showing only diff rows then don't display the count column's diff row as a diff for formatting
 				return null;
@@ -972,7 +972,7 @@ namespace PerfSummaries
 
 				string prefix = "";
 				bool bIsDiffRow = IsDiffRow(index);
-	        	if (bShowOnlyDiffRows && isCountColumn)
+	        	if (!bShowOnlyDiffRows && isCountColumn)
 				{
 					// If we're showing only diff rows then don't display the count column's diff row as a diff for formatting
 					bIsDiffRow=false;
