PerfReportTool - make summaryTableXmlSubst case-insensitive, consistent with the filter list

diff --git a/Engine/Source/Programs/CSVTools/PerfReportTool/ReportXml.cs b/Engine/Source/Programs/CSVTools/PerfReportTool/ReportXml.cs
--- a/Engine/Source/Programs/CSVTools/PerfReportTool/ReportXml.cs
+++ b/Engine/Source/Programs/CSVTools/PerfReportTool/ReportXml.cs
@@ -326,7 +326,7 @@ namespace PerfReportTool
 						string [] pair = substStr.Split('=');
 						if (pair.Length == 2)
 						{
-							substitutionsDict[pair[0]] = pair[1].Split(",").ToList();
+							substitutionsDict[pair[0].ToLowerInvariant()] = pair[1].Split(",").ToList();
 						}
 					}
 				}
diff --git a/Engine/Source/Programs/CSVTools/PerfReportTool/SummaryTable.cs b/Engine/Source/Programs/CSVTools/PerfReportTool/SummaryTable.cs
--- a/Engine/Source/Programs/CSVTools/PerfReportTool/SummaryTable.cs
+++ b/Engine/Source/Programs/CSVTools/PerfReportTool/SummaryTable.cs
@@ -147,7 +147,7 @@ namespace PerfSummaries
 			}
 			for (int i = 0; i < list.Count; i++)
 			{
-				if (substitutionsDict.TryGetValue(list[i], out List<string> replaceList))
+				if (substitutionsDict.TryGetValue(list[i].ToLowerInvariant(), out List<string> replaceList))
 				{
 					list.RemoveAt(i);
 					list.InsertRange(i, replaceList);
@@ -162,7 +162,7 @@ namespace PerfSummaries
 			{
 				return str;
 			}
-			if (substitutionsDict.TryGetValue(str, out List<string> replaceList))
+			if (substitutionsDict.TryGetValue(str.ToLowerInvariant(), out List<string> replaceList))
 			{
 				if (replaceList.Count>1 || replaceList.Count == 0)
 				{
