Minor bugfix to metadataProxy handling
Extraneous warning emitted when proxy is found and there are no overrides

#rnx
#tests PRT local

diff --git a/Engine/Source/Programs/CSVTools/PerfReportTool/PerfReportTool.cs b/Engine/Source/Programs/CSVTools/PerfReportTool/PerfReportTool.cs
--- a/Engine/Source/Programs/CSVTools/PerfReportTool/PerfReportTool.cs
+++ b/Engine/Source/Programs/CSVTools/PerfReportTool/PerfReportTool.cs
@@ -23,7 +23,7 @@ namespace PerfReportTool
     class Version
     {
 		// Format: Major.Minor.Bugfix
-        private static string VersionString = "4.240.0";
+        private static string VersionString = "4.240.1";
 
         public static string Get() { return VersionString; }
     };
@@ -351,13 +351,13 @@ namespace PerfReportTool
 					if (metadataProxyName.Length > 0)
 					{
 						proxyMetadata = reportXML.GetMetadataProxyInfo(metadataProxyName);
-						if (explicitMetadata.Length > 0)
+						if (proxyMetadata.Length == 0)
 						{
-							proxyMetadata += ",";
+							Console.WriteLine("Warning: Unable to find metadataProxy named " + metadataProxyName + " in report xml csvMetadataProxies section.");
 						}
-						else
+						else if (explicitMetadata.Length > 0)
 						{
-							Console.WriteLine("Warning: Unable to find metadataProxy named " + metadataProxyName + " in report xml csvMetadataProxies section.");
+							proxyMetadata += ",";
 						}
 					}
 
