PRT - Add option to skip setting metadata in csvconvert if it doesn't already have metadata.

This prevents setting some metadata on csvs that previously had none, resulting in partial metadata that is invalid for most of our tools.

#rnx

diff --git a/Engine/Source/Programs/CSVTools/CsvConvert/CsvConvert.cs b/Engine/Source/Programs/CSVTools/CsvConvert/CsvConvert.cs
--- a/Engine/Source/Programs/CSVTools/CsvConvert/CsvConvert.cs
+++ b/Engine/Source/Programs/CSVTools/CsvConvert/CsvConvert.cs
@@ -66,6 +66,7 @@ namespace CSVInfo
 				"  [-inPlace] - overwrite the input file (implies -force)\n" +
 				"  [-verify]\n" +
 				"  [-dontFixMissingID] - doesn't generate a CsvID if it's not found in the metadata\n" +
+				"  [-skipIfMetadataMissing] - don't set metadata if the csv doesn't have any.\n" +
 				"  [-setMetadata key0=value0;key1=value1;...]\n";
 
 
@@ -190,15 +191,24 @@ namespace CSVInfo
 			string SetMetadataStr = GetArg("setMetadata");
 
 			// Make sure we have metadata if we're setting it
+			bool bSkipUpdatingMetadata = false;
 			if (csvStats.metaData == null && (setMetadataFile != null || SetMetadataStr != null) )
 			{
-				Console.WriteLine("File has null metadata. Creating...");
-				csvStats.metaData = new CsvMetadata();
-				csvStats.GenerateCsvIDIfMissing();
+				if (GetBoolArg("skipIfMetadataMissing"))
+				{
+					Console.WriteLine("Skipping writing metadata since it has none.");
+					bSkipUpdatingMetadata = true;
+				}
+				else
+				{
+					Console.WriteLine("File has null metadata. Creating...");
+					csvStats.metaData = new CsvMetadata();
+					csvStats.GenerateCsvIDIfMissing();
+				}
 			}
 
 
-			if (setMetadataFile != null)
+			if (setMetadataFile != null && !bSkipUpdatingMetadata)
 			{
 				if (csvStats.metaData == null)
 				{
@@ -218,7 +228,7 @@ namespace CSVInfo
 				}
 			}
 
-			if (SetMetadataStr != null && SetMetadataStr.Length > 0)
+			if (SetMetadataStr != null && SetMetadataStr.Length > 0 && !bSkipUpdatingMetadata)
 			{
 				bool bModified = false;
 				if (csvStats.metaData == null)
