Csv tools: Add support for derived csv stats using arbitrary expressions. Derived stats can be defined as inline expressions in CsvToSVG -stats args or in PerfReportTool's <DerivedCsvStats> element.

CsvToSVG 3.65:
- When passing inline stats to CsvToSVG, the expression must be surrounded by brackets.
- Pass <DerivedStatName>=<expression> to give stats a name. e.g -stats=(AllWidgets=WidgetsTypeA+WidgetsTypeB)

PerfReportTool:
- Add derived CSV stats in the <derivedCsvStats> element

Surround stat names with single quotes to resolve ambiguity. This is necessary for stat names with slashes. e.g 'exclusive/gamethread/ui'

diff --git a/Engine/Source/Programs/CSVTools/CSVToSVG/CsvToSvgProgram.cs b/Engine/Source/Programs/CSVTools/CSVToSVG/CsvToSvgProgram.cs
--- a/Engine/Source/Programs/CSVTools/CSVToSVG/CsvToSvgProgram.cs
+++ b/Engine/Source/Programs/CSVTools/CSVToSVG/CsvToSvgProgram.cs
@@ -442,15 +442,26 @@ namespace CSVTools
         {
             int numRowsToSkip = GetIntArg("skipRows",0);
 
+			// If we have derived stat expressions then load all stats
+			bool bHasDerivedStatExpression = false;
+			foreach (string stat in statNames)
+			{
+				if ( CsvStats.IsDerivedStatExpression(stat) )
+				{
+					bHasDerivedStatExpression = true; 
+					break;
+				}
+			}
+
 			// Read the file from the cache if we have one
 			CsvStats csvStats = null;
 			if (csvCache != null)
 			{
-				csvStats = csvCache.ReadCSVFile(csvFilename, statNames, numRowsToSkip);
+				csvStats = csvCache.ReadCSVFile(csvFilename, bHasDerivedStatExpression ? null : statNames, numRowsToSkip);
 			}
 			else
 			{
-				csvStats = CsvStats.ReadCSVFile(csvFilename, statNames, numRowsToSkip);
+				csvStats = CsvStats.ReadCSVFile(csvFilename, bHasDerivedStatExpression ? null : statNames, numRowsToSkip);
 			}
             return csvStats;
         }
diff --git a/Engine/Source/Programs/CSVTools/CsvStats/CommandLineTool.cs b/Engine/Source/Programs/CSVTools/CsvStats/CommandLineTool.cs
--- a/Engine/Source/Programs/CSVTools/CsvStats/CommandLineTool.cs
+++ b/Engine/Source/Programs/CSVTools/CsvStats/CommandLineTool.cs
@@ -309,7 +309,7 @@ namespace CSVStats
 			{
 				listStr = listStr.ToLower();
 			}
-			return listStr.Split(separator).ToList();
+			return CsvStats.SplitStringListWithBracketGroups(listStr, separator);
 		}
 
 		protected void WriteLine(String message, params object[] args)
diff --git a/Engine/Source/Programs/CSVTools/CsvStats/CsvStats.cs b/Engine/Source/Programs/CSVTools/CsvStats/CsvStats.cs
--- a/Engine/Source/Programs/CSVTools/CsvStats/CsvStats.cs
+++ b/Engine/Source/Programs/CSVTools/CsvStats/CsvStats.cs
@@ -6,13 +6,13 @@ using System.Collections.Generic;
 using System.Linq;
 using System.Text;
 using System.Threading.Tasks;
-using System.Drawing;
 using System.IO;
 using System.Globalization;
-using System.Diagnostics;
 using System.Collections;
 using System.IO.Compression;
 using System.Security.Cryptography;
+using System.Data;
+using System.Linq.Expressions;
 
 namespace CSVStats
 {
@@ -501,6 +501,256 @@ namespace CSVStats
 			}
 			return null;
 		}
+
+		private string ExpandCustomFunctions(string expression)
+		{
+			char[] operators = { '+', '-', '*', '/' };
+
+			string[] customFunctionNames = { "add" };
+			foreach (string funcName in customFunctionNames)
+			{
+				int startIdx = expression.IndexOf(funcName+"(");
+				if (startIdx == -1)
+				{
+					continue;
+				}
+				if (startIdx > 0 && !operators.Contains(expression[startIdx - 1]))
+				{
+					continue;
+				}
+				// FIXME: This will fail with nested brackets
+				int endIdx = expression.IndexOf(')', startIdx+1);
+				if (endIdx == -1)
+				{
+					continue;
+				}
+				string inner = expression[(startIdx+funcName.Length+1)..(endIdx)];
+				string[] innerParams = inner.Split(',');
+				if (funcName=="add")
+				{
+					// add(a,b,c) expands to (a+b+c). This is useful with wildcard expansion, since wildcards expand to comma separated list
+					string replaceStr = "(" + String.Join("+", innerParams) + ")";
+					expression = expression[..startIdx]+replaceStr+expression[(endIdx+1)..];
+				}
+				// Other functions can go here
+			}
+			return expression;
+		}
+
+		public StatSamples MakeDerivedStatFromExpression(string expression)
+		{
+			if (expression == "")
+			{
+				return null;
+			}
+
+			// Remove whitespace
+			expression = expression.Replace(" ", "");
+
+			// Remove outer brackets if they exist
+			if (expression[0] == '(' && expression[^1]==')' )
+			{
+				expression = expression[1..^1];
+			}
+
+			// If we have an equals, to the left of it is the stat name
+			int equalsIndex = expression.IndexOf('=');
+			string statName;
+			if (equalsIndex != -1)
+			{
+				statName = expression[..equalsIndex];
+				expression = expression[(equalsIndex + 1)..];
+			}
+			else
+			{
+				statName = expression.Replace("'", " ");
+			}
+
+			// Check for collisions
+			if (GetStat(statName) != null) 
+			{
+				Console.WriteLine("Derived expression stat name already exists: " + statName + ". Ignoring.");
+				return null;
+			}
+
+			expression = expression.ToLower();
+
+			Dictionary<string, StatSamples> referencedStats = new Dictionary<string, StatSamples>();
+
+			// Find stats in single quotes and generate sanitized names for them (ie replace / with _SLASH_)
+			if (expression.Contains('\''))
+			{
+				string newExpression = "";
+				bool bInQuotes = false;
+				string token = "";
+				for (int i=0; i<expression.Length; i++) 
+				{
+					char c = expression[i];
+					if (c == '\'') 
+					{
+						bInQuotes = !bInQuotes;
+						if ( bInQuotes )
+						{
+							token = "";
+						}
+						else
+						{
+							if (token.Length>0)
+							{
+								// Find the matching stat. If a wildcard is specified, multiple stats may be returned
+								List<StatSamples> matchingStats = GetStatsMatchingString(token);
+								if (matchingStats.Count == 0)
+								{
+									Console.WriteLine("Warning: stat not found: " + token);
+									return null;
+								}
+								bool bFirst = true;
+								foreach (StatSamples stat in matchingStats)
+								{
+									string sanitizedStatName = stat.Name.ToLower().Replace("/", "_SLASH_");
+									if (!bFirst)
+									{
+										newExpression += ',';
+									}
+									newExpression += sanitizedStatName;
+									referencedStats[sanitizedStatName] = stat;
+									bFirst = false;
+								}
+							}
+						}
+					}
+					else
+					{
+						if (bInQuotes)
+						{
+							token += c;
+						}
+						else
+						{
+							newExpression += c;
+						}
+					}
+				}
+				expression = newExpression;
+			}
+
+			// Handle special functions we might want to expand
+			expression = ExpandCustomFunctions(expression);
+
+			// Find any other refererenced stats
+			char[] delimiters = { '+', '-', '*', '(', ')', '/' };
+			string[] tokens = expression.Split(delimiters);
+			foreach (string token in tokens)
+			{
+				if (token == "")
+				{
+					continue;
+				}
+				if (!referencedStats.ContainsKey(token))
+				{
+					StatSamples stat = GetStat(token);
+					if (stat != null)
+					{
+						referencedStats.Add(token, stat);
+					}
+				}
+			}
+
+			// Construct the datatable with any referenced stats
+			DataTable dataTable = new DataTable();
+			foreach (KeyValuePair<string, StatSamples> entry in referencedStats)
+			{
+				dataTable.Columns.Add(entry.Key, typeof(float));
+			}
+			StatSamples[] statsList = referencedStats.Values.ToArray();
+			for (int i=0; i<SampleCount; i++)
+			{
+				DataRow row = dataTable.NewRow();
+				for (int j=0; j<statsList.Length; j++)
+				{
+					row[j] = statsList[j].samples[i];
+				}
+				dataTable.Rows.Add(row);
+			}
+
+			// Add the result column
+			int resultColumnIndex = dataTable.Columns.Count;
+			DataColumn resultColumn = new DataColumn("__RESULT__", typeof(float));
+
+			try
+			{
+				resultColumn.Expression = expression;
+				dataTable.Columns.Add(resultColumn);
+			}
+			catch (Exception ex) 
+			{ 
+				Console.WriteLine("Error: derived stat expression \""+expression+"\"could not be computed."+ex.ToString());
+				return null;
+			}
+			dataTable.AcceptChanges();
+			
+			// Write out the derived stat
+			StatSamples statSamples = new StatSamples(statName);
+			foreach (DataRow row in dataTable.Rows)
+			{
+				statSamples.samples.Add((float)row[resultColumnIndex]);
+			}
+			statSamples.ComputeAverageAndTotal();
+
+			return statSamples;
+		}
+
+		public StatSamples AddDerivedStatFromExpression(string expression)
+		{
+			StatSamples derivedStat = MakeDerivedStatFromExpression(expression);
+			if (derivedStat != null)
+			{
+				AddStat(derivedStat);
+			}
+			return derivedStat;
+		}
+
+		// Splits a list of strings but preserves tokens inside brackets
+		public static List<string> SplitStringListWithBracketGroups(string listStr, char separator)
+		{
+			if (!listStr.Contains('('))
+			{
+				// If there are no brackets, we can just split
+				return listStr.Split(separator).Select(token => token.Trim()).ToList();
+			}
+			int bracketDepth = 0;
+			string token = "";
+			List<string> listOut = new List<string>();
+			for (int i=0; i<listStr.Length; i++) 
+			{ 
+				char c = listStr[i];
+				if (c == '(')
+				{
+					bracketDepth++;
+				}
+				else if (c == ')')
+				{
+					bracketDepth--;
+					bracketDepth = Math.Max(0, bracketDepth);
+				}
+				if (bracketDepth==0 && c==separator)
+				{
+					listOut.Add(token.Trim());
+					token = "";
+				}
+				else
+				{
+					token += c;
+				}
+			}
+			if (token != "")
+			{
+				listOut.Add(token.Trim());
+			}
+			return listOut;
+		}
+
+
 		public void AddStat(StatSamples stat)
 		{
 			Stats.Add(stat.Name.ToLower(), stat);
@@ -1682,8 +1932,13 @@ namespace CSVStats
 			return expressionTree.Evaluate(metadata);
 		}
 
+		public static bool IsDerivedStatExpression(string stat)
+		{
+			return stat.StartsWith('(') && stat.EndsWith(')') || stat.Contains('=');
+		}
+
 
-		public static CsvStats ReadCSVFile(string csvFilename, string[] statNames, int numRowsToSkip = 0, bool bGenerateCsvIdIfMissing=false, CsvFileInfo FileInfoOut=null, bool bJustHeader=false)
+		public static CsvStats ReadCSVFile(string csvFilename, string[] statNames=null, int numRowsToSkip = 0, bool bGenerateCsvIdIfMissing=false, CsvFileInfo FileInfoOut=null, bool bJustHeader=false)
         {
 			CsvStats statsOut;
 			if (csvFilename.EndsWith(".csv.bin"))
@@ -1777,7 +2032,7 @@ namespace CSVStats
             return false;
         }
 
-        public static CsvStats ReadCSVFromLines(string[] linesArray, string[] statNames, int numRowsToSkip = 0, bool skipReadingData=false, CsvFileInfo fileInfo=null)
+        public static CsvStats ReadCSVFromLines(string[] linesArray, string[] statNames=null, int numRowsToSkip = 0, bool skipReadingData=false, CsvFileInfo fileInfo=null)
         {
             List<string> lines = linesArray.ToList();
 
diff --git a/Engine/Source/Programs/CSVTools/CsvToSvgLib/CsvToSvgLib.cs b/Engine/Source/Programs/CSVTools/CsvToSvgLib/CsvToSvgLib.cs
--- a/Engine/Source/Programs/CSVTools/CsvToSvgLib/CsvToSvgLib.cs
+++ b/Engine/Source/Programs/CSVTools/CsvToSvgLib/CsvToSvgLib.cs
@@ -15,7 +15,7 @@ namespace CSVTools
 
 	public class CsvToSvgLibVersion
 	{
-		private static string VersionString = "3.64";
+		private static string VersionString = "3.65";
 
 		public static string Get() { return VersionString; }
 	};
@@ -2708,6 +2708,20 @@ namespace CSVTools
 		CsvStats ProcessCsvStats(CsvStats csvStatsIn, GraphParams graphParams, out FrameRange frameRange)
 		{
 			CsvStats csvStats = new CsvStats(csvStatsIn, graphParams.statNames.ToArray());
+
+			// Compute any derived stats
+			foreach (string statStr in graphParams.statNames)
+			{
+				if (CsvStats.IsDerivedStatExpression(statStr))
+				{
+					StatSamples derivedStat = csvStatsIn.MakeDerivedStatFromExpression(statStr);
+					if (derivedStat != null)
+					{
+						csvStats.AddStat(derivedStat);
+					}
+				}
+			}
+
 			if (graphParams.discardLastFrame)
 			{
 				foreach (StatSamples stat in csvStats.Stats.Values.ToArray())
diff --git a/Engine/Source/Programs/CSVTools/PerfReportTool/PerfReportTool.cs b/Engine/Source/Programs/CSVTools/PerfReportTool/PerfReportTool.cs
--- a/Engine/Source/Programs/CSVTools/PerfReportTool/PerfReportTool.cs
+++ b/Engine/Source/Programs/CSVTools/PerfReportTool/PerfReportTool.cs
@@ -1027,6 +1027,17 @@ namespace PerfReportTool
 				perfLog.LogTiming("    ReadCsvStats");
 			}
 
+			// Compute derived CSV stats
+			foreach (DerivedCsvStatDefinition derivedStat in reportXML.derivedCsvStatDefinitions.Values)
+			{
+				csvStats.AddDerivedStatFromExpression(derivedStat.getExpressionStr());
+			}
+			if (perfLog != null)
+			{
+				perfLog.LogTiming("    Compute Derived CSV stats");
+			}
+
+
 			if (bStripStatsByEvents)
 			{
 				CsvStats strippedCsvStats = StripCsvStatsByEvents(csvStatsUnstripped, out numFramesStripped);
@@ -1147,6 +1158,28 @@ namespace PerfReportTool
 						throw new Exception("Graph " + graph.title + " has no <statString> element");
 					}
 
+					// Replace any referenced derived stats with the full expression so that CSVToSVG can expand them
+					// TODO: Technically this should only be necessary if CSVToSVG is running as a separate process
+					if (reportXML.derivedCsvStatDefinitions.Count > 0)
+					{
+						string[] stats = graph.settings.statString.value.ToLowerInvariant().Split(',');
+						bool bHasDerivedStats = false;
+						for (int i = 0; i < stats.Length; i++)
+						{
+							DerivedCsvStatDefinition derivedStatDefinition = reportXML.GetDerivedCsvStatDefinition(stats[i].Trim());
+							if (derivedStatDefinition != null)
+							{
+								stats[i] = derivedStatDefinition.getExpressionStr();
+								bHasDerivedStats = true;
+							}
+						}
+						if (bHasDerivedStats) 
+						{
+							graph.settings.statString.value = String.Join(',', stats);
+						}
+					}
+
+
 					bool bShowGraph = graph.ShouldShow(csvFile.finalCsv ?? csvFile.dummyCsvStats);
 					if (bShowGraph)
 					{
@@ -1946,6 +1979,8 @@ namespace PerfReportTool
 			AddConditionalArg("endEvent", graphSettings.endEvent.isSet, graphSettings.endEvent.value);
 			AddConditionalArg("endEventOffset", graphSettings.endEventOffset.isSet, graphSettings.endEventOffset.value);
 
+			List<string> statList = CsvStats.SplitStringListWithBracketGroups(graphSettings.statString.value,',');
+
 			string argString = string.Empty;
 			if (argFormat == CsvToSvgArgFormat.Url)
 			{
@@ -1959,7 +1994,8 @@ namespace PerfReportTool
 				}
 				// Exclude csvs from the args as that is a separate argument that the caller must setup.
 				// Any multi-value args need to be delimited with a semi-colon.
-				args["stats"] = FormatStringList(graphSettings.statString.value, ",");
+				// Derived csv stat definitions can include equals, so encode it
+				args["stats"] = String.Join(";", statList).Replace("=", "%3D"); 
 				AddConditionalArg("showevents", showEvents.Length > 0, FormatStringList(showEvents, " "));
 				AddConditionalArg("highlightEventRegions", highlightEventRegions.Length > 0, FormatStringList(highlightEventRegions, ","));
 
@@ -1977,8 +2013,7 @@ namespace PerfReportTool
 				AddConditionalArg("showevents", showEvents.Length > 0, showEvents);
 				AddConditionalArg("highlightEventRegions", highlightEventRegions.Length > 0, highlightEventRegions);
 
-				string[] statStringTokens = graphSettings.statString.value.Split(',');
-				IEnumerable<string> quoteWrappedStatStrings = statStringTokens.Select(token => '"' + token + '"');
+				IEnumerable<string> quoteWrappedStatStrings = statList.Select(token => '"' + token + '"');
 				args["stats"] = String.Join(" ", quoteWrappedStatStrings);
 
 				List<string> argList = args.Select(a =>
@@ -1998,7 +2033,7 @@ namespace PerfReportTool
 			graphParams.title = graph.title;
 
 			GraphSettings graphSettings = graph.settings;
-			graphParams.statNames.AddRange(graphSettings.statString.value.Split(','));
+			graphParams.statNames.AddRange(CsvStats.SplitStringListWithBracketGroups(graphSettings.statString.value,','));
 			graphParams.lineThickness = (float)(graphSettings.thickness.value * thicknessMultiplier);
 			graphParams.smooth = graphSettings.smooth.value && !GetBoolArg("nosmooth");
 			if (graphParams.smooth)
@@ -2161,8 +2196,15 @@ namespace PerfReportTool
 			string csvToolPath = Path.Combine(GetBaseDirectory(), "CSVToSVG.exe");
 			string binary = csvToolPath;
 
+			if (Host == HostPlatform.Windows)
+			{
+				if (!File.Exists(binary) && binary.ToLower().Contains("\\perfreporttool\\"))
+				{
+					binary = binary.Replace("\\perfreporttool\\", "\\CsvToSVG\\", StringComparison.InvariantCultureIgnoreCase);
+				}
+			}
 			// run mono on non-Windows hosts
-			if (Host != HostPlatform.Windows)
+			else
 			{
 				// note, on Mac mono will not be on path
 				binary = Host == HostPlatform.Linux ? "mono" : "/Library/Frameworks/Mono.framework/Versions/Current/Commands/mono";
diff --git a/Engine/Source/Programs/CSVTools/PerfReportTool/ReportXml.cs b/Engine/Source/Programs/CSVTools/PerfReportTool/ReportXml.cs
--- a/Engine/Source/Programs/CSVTools/PerfReportTool/ReportXml.cs
+++ b/Engine/Source/Programs/CSVTools/PerfReportTool/ReportXml.cs
@@ -61,6 +61,21 @@ namespace PerfReportTool
 		public bool bHasVariables;
 	};
 
+	class DerivedCsvStatDefinition
+	{
+		public DerivedCsvStatDefinition(XElement element, XmlVariableMappings vars)
+		{
+			name = element.GetSafeAttribute<string>(vars, "name");
+			expression = element.GetSafeAttribute<string>(vars, "expression");
+		}
+		public string getExpressionStr()
+		{
+			return name + "=" + expression;
+		}
+		public string name;
+		string expression;
+	}
+
 	class DerivedMetadataMappings
 	{
 		public DerivedMetadataMappings()
@@ -293,6 +308,19 @@ namespace PerfReportTool
 				}
 			}
 
+			// Read the derived CSV stat definitions mappings
+			derivedCsvStatDefinitions = new Dictionary<string, DerivedCsvStatDefinition>();
+			XElement derivedCsvStatsElement = rootElement.Element("derivedCsvStats");
+			if (derivedCsvStatsElement != null)
+			{
+				foreach (XElement el in derivedCsvStatsElement.Elements("derivedCsvStat"))
+				{
+					DerivedCsvStatDefinition derivedCsvStatDefinition = new DerivedCsvStatDefinition(el, staticVariableMappings);
+					derivedCsvStatDefinitions[derivedCsvStatDefinition.name.ToLowerInvariant()] = derivedCsvStatDefinition;
+				}
+			}
+
+
 			// Read events to strip
 			XElement eventsToStripEl = rootElement.Element("csvEventsToStrip");
 			if (eventsToStripEl != null)
@@ -610,6 +638,15 @@ namespace PerfReportTool
 			return "";
 		}
 
+		public DerivedCsvStatDefinition GetDerivedCsvStatDefinition(string name)
+		{
+			if (derivedCsvStatDefinitions.TryGetValue(name, out DerivedCsvStatDefinition def))
+			{
+				return def;
+			}
+			return null; 
+		}
+
 		Dictionary<string, SummaryTableInfo> summaryTables;
 
 		XElement reportTypesElement;
@@ -629,6 +666,7 @@ namespace PerfReportTool
 		List<CsvEventStripInfo> csvEventsToStrip;
 		string reportTypeXmlFilename;
 		public DerivedMetadataMappings derivedMetadataMappings;
+		public Dictionary<string, DerivedCsvStatDefinition> derivedCsvStatDefinitions;
 	}
 
 	class XmlVariableMappings
