PerfReportTool 4.241.0 - Improve handling of derived CSV stats.
- Compute derived stats at the start of the process so they're available for graphing and work with wildcard searches etc.
- Only pass the full expressions to CsvToSvg if it's is being run in a separate process (or via URL). This avoids us computing the derived stats multiple times
- Cleanup some helper functions for querying stats by wildcard
- Improve the robustness of derived stat expansion for graphing. Make sure this works with wildcards
- Add support for passing lists of stat names  in quotes as part of an expression separated by commas. This makes interop with variables easier, since we can wrap a list variable in quotes and make sure the stat names are passed literally

diff --git a/Engine/Source/Programs/CSVTools/CsvStats/CsvStats.cs b/Engine/Source/Programs/CSVTools/CsvStats/CsvStats.cs
--- a/Engine/Source/Programs/CSVTools/CsvStats/CsvStats.cs
+++ b/Engine/Source/Programs/CSVTools/CsvStats/CsvStats.cs
@@ -286,7 +286,7 @@ namespace CSVStats
             colour = source.colour;
             LegendName = source.LegendName;
         }
-        public StatSamples(string name, int sampleListLength=0)
+        public StatSamples(string name, int sampleListLength=0, bool bInIsDerived = false)
         {
             Name = name;
             LegendName = name;
@@ -301,7 +301,8 @@ namespace CSVStats
             average = 0.0f;
             total = 0.0f;
             colour = new Colour(0, 0.0f);
-        }
+			bIsDerived = bInIsDerived;
+		}
 
         public float GetRatioOfFramesInBudget(float budget)
         {
@@ -447,6 +448,7 @@ namespace CSVStats
         public float average;
         public double total;
         public Colour colour;
+		public bool bIsDerived;
     };
 
 
@@ -578,6 +580,7 @@ namespace CSVStats
 			Dictionary<string, StatSamples> referencedStats = new Dictionary<string, StatSamples>();
 
 			// Find stats in single quotes and generate sanitized names for them (ie replace / with _SLASH_)
+			// This works for individual stats or for multiple stats separated by commas
 			if (expression.Contains('\''))
 			{
 				string newExpression = "";
@@ -598,7 +601,7 @@ namespace CSVStats
 							if (token.Length>0)
 							{
 								// Find the matching stat. If a wildcard is specified, multiple stats may be returned
-								List<StatSamples> matchingStats = GetStatsMatchingString(token);
+								List<StatSamples> matchingStats = GetStatsMatchingStringList(token.Split(','));
 								if (matchingStats.Count == 0)
 								{
 									Console.WriteLine("Warning: stat not found: " + token);
@@ -673,27 +676,27 @@ namespace CSVStats
 				dataTable.Rows.Add(row);
 			}
 
-			// Add the result column
+			// Add the result column using the expression to compute
 			int resultColumnIndex = dataTable.Columns.Count;
 			DataColumn resultColumn = new DataColumn("__RESULT__", typeof(float));
-
 			try
 			{
 				resultColumn.Expression = expression;
 				dataTable.Columns.Add(resultColumn);
+				dataTable.AcceptChanges();
 			}
 			catch (Exception ex) 
 			{ 
 				Console.WriteLine("Error: derived stat expression \""+expression+"\"could not be computed."+ex.ToString());
 				return null;
 			}
-			dataTable.AcceptChanges();
-			
+
 			// Write out the derived stat
-			StatSamples statSamples = new StatSamples(statName);
-			foreach (DataRow row in dataTable.Rows)
+			StatSamples statSamples = new StatSamples(statName, dataTable.Rows.Count, true);
+			for (int i=0; i<dataTable.Rows.Count; i++)
 			{
-				statSamples.samples.Add((float)row[resultColumnIndex]);
+				DataRow row = dataTable.Rows[i];
+				statSamples.samples[i] = (float)row[resultColumnIndex];
 			}
 			statSamples.ComputeAverageAndTotal();
 
@@ -713,27 +716,39 @@ namespace CSVStats
 		// Splits a list of strings but preserves tokens inside brackets
 		public static List<string> SplitStringListWithBracketGroups(string listStr, char separator)
 		{
-			if (!listStr.Contains('('))
+			if (!listStr.Contains('(') && !listStr.Contains('\''))
 			{
-				// If there are no brackets, we can just split
+				// If there are no brackets or quotes, we can just split
 				return listStr.Split(separator).Select(token => token.Trim()).ToList();
 			}
 			int bracketDepth = 0;
+			bool bInQuotes = false;
 			string token = "";
 			List<string> listOut = new List<string>();
 			for (int i=0; i<listStr.Length; i++) 
 			{ 
 				char c = listStr[i];
-				if (c == '(')
+
+				// If we're not in brackets then match quotes
+				if (bracketDepth == 0 && c == '\'')
 				{
-					bracketDepth++;
+					bInQuotes = !bInQuotes;
 				}
-				else if (c == ')')
+				// If we're not in quotes then match brackets
+				if (!bInQuotes)
 				{
-					bracketDepth--;
-					bracketDepth = Math.Max(0, bracketDepth);
+					if (c == '(')
+					{
+						bracketDepth++;
+					}
+					else if (c == ')')
+					{
+						bracketDepth--;
+						bracketDepth = Math.Max(0, bracketDepth);
+					}
 				}
-				if (bracketDepth==0 && c==separator)
+
+				if (!bInQuotes && bracketDepth==0 && c==separator)
 				{
 					listOut.Add(token.Trim());
 					token = "";
@@ -993,27 +1008,41 @@ namespace CSVStats
 			return results;
 		}
 
-
-		public Dictionary<string, bool> GetStatNamesMatchingStringList_Dict(string[] statNames)
+		public HashSet<StatSamples> GetStatsMatchingStringList_Set(string[] statNames)
 		{
-			Dictionary<string, bool> uniqueStatsDict = new Dictionary<string, bool>();
+			HashSet<StatSamples> uniqueStats = new HashSet<StatSamples>();
 			foreach (string statStr in statNames)
 			{
 				List<StatSamples> statsMatching = GetStatsMatchingString(statStr);
 				foreach (StatSamples stat in statsMatching)
 				{
-					if (!uniqueStatsDict.ContainsKey(stat.Name))
+					if (!uniqueStats.Contains(stat))
 					{
-						uniqueStatsDict.Add(stat.Name, true);
+						uniqueStats.Add(stat);
 					}
 				}
 			}
-			return uniqueStatsDict;
+			return uniqueStats;
+		}
+		public List<StatSamples> GetStatsMatchingStringList(string[] statNames)
+		{
+			return GetStatsMatchingStringList_Set(statNames).ToList();
+		}
+
+		public HashSet<string> GetStatNamesMatchingStringList_Set(string[] statNames)
+		{
+			HashSet<StatSamples> uniqueStats = GetStatsMatchingStringList_Set(statNames);
+			HashSet<string> statNamesOut = new HashSet<string>();
+			foreach (StatSamples stat in uniqueStats)
+			{
+				statNamesOut.Add(stat.Name);
+			}
+			return statNamesOut;
 		}
 
 		public List<string> GetStatNamesMatchingStringList(string[] statNames)
 		{
-			return GetStatNamesMatchingStringList_Dict(statNames).Keys.ToList();
+			return GetStatNamesMatchingStringList_Set(statNames).ToList();
 		}
 
 		public List<StatSamples> GetStatsMatchingString(string statString)
@@ -1249,10 +1278,10 @@ namespace CSVStats
 			}
 
 			// Get the list of all stats
-			Dictionary<string, bool> statNamesToReadDict = null;
+			HashSet<string> statNamesToReadSet = null;
 			if (statNamesToRead != null)
 			{
-				statNamesToReadDict = csvStatsOut.GetStatNamesMatchingStringList_Dict(statNamesToRead);
+				statNamesToReadSet = csvStatsOut.GetStatNamesMatchingStringList_Set(statNamesToRead);
 			}
 
 			// Fill in the file info
@@ -1282,7 +1311,7 @@ namespace CSVStats
 				double statTotal = fileReader.ReadDouble();
 
 				int statSizeBytes = fileReader.ReadInt32();
-				bool readThisStat = statNamesToReadDict == null || statNamesToReadDict.ContainsKey(statName);
+				bool readThisStat = statNamesToReadSet == null || statNamesToReadSet.Contains(statName);
 				if (readThisStat == false)
 				{
 					// If we're skipping this stat then just read the bytes and remove it from the stat list
diff --git a/Engine/Source/Programs/CSVTools/PerfReportTool/CsvFileCache..cs b/Engine/Source/Programs/CSVTools/PerfReportTool/CsvFileCache..cs
--- a/Engine/Source/Programs/CSVTools/PerfReportTool/CsvFileCache..cs
+++ b/Engine/Source/Programs/CSVTools/PerfReportTool/CsvFileCache..cs
@@ -27,7 +27,7 @@ namespace PerfReportTool
 			BinaryCsv,
 			SummaryTableCachePrc
 		};
-		public CachedCsvFile(string inFilename, bool useCacheFiles, DerivedMetadataMappings derivedMetadataMappings)
+		public CachedCsvFile(string inFilename, bool useCacheFiles, DerivedMetadataMappings inDerivedMetadataMappings, DerivedCsvStatDefinition[] inDerivedCsvStatDefinitions)
 		{
 			if (inFilename.ToLower().EndsWith(".csv"))
 			{
@@ -110,33 +110,48 @@ namespace PerfReportTool
 			}
 			filename = inFilename;
 
-			// Setup initial metadata variable mappings and derived metadata (which can read from the initial variableMappings). 
+			// Setup initial metadata variable mappings derived metadata (which can read from the initial variableMappings). 
 			// Note: we skip for PRCs (dummyCsvStats==null), since these already have mappings cached
 			if (dummyCsvStats != null)
 			{
 				xmlVariableMappings = new XmlVariableMappings();
 
-				// Note: We could apply the global variable set here, but that is slow and leads to complexity and correctness issues
+				// Note: We could apply the global variable set here, but that is slow and leads to complexity and correctness issues due to dependency order
 				if (dummyCsvStats.metaData != null)
 				{
 					metadata = dummyCsvStats.metaData;
 					xmlVariableMappings.SetMetadataVariables(metadata);
-					derivedMetadataMappings.ApplyMapping(metadata, xmlVariableMappings);
+					inDerivedMetadataMappings.ApplyMapping(metadata, xmlVariableMappings);
+				}
+
+				// Add the derived csv stats to dummyStats
+				derivedCsvStatDefinitions = inDerivedCsvStatDefinitions;
+				foreach (DerivedCsvStatDefinition derivedCsvStatDefinition in derivedCsvStatDefinitions)
+				{
+					dummyCsvStats.AddStat(new StatSamples(derivedCsvStatDefinition.name, 0, true));
 				}
 			}
 		}
 
+		private void AddDerivedMetadataToCsv(CsvStats csvStats)
+		{
+			foreach (DerivedCsvStatDefinition derivedStat in derivedCsvStatDefinitions)
+			{
+				finalCsv.AddDerivedStatFromExpression(derivedStat.getExpressionStr());
+			}
+		}
+
 		public void PrepareCsvData()
 		{
 			if (fileType == FileType.BinaryCsv)
 			{
 				finalCsv = CsvStats.ReadBinFile(filename);
+				AddDerivedMetadataToCsv(finalCsv);
 				// If we already have metadata then use it. This avoids the need to remap it multiple times
 				if (metadata != null)
 				{
 					finalCsv.metaData = metadata;
 				}
-
 			}
 			else
 			{
@@ -155,20 +170,21 @@ namespace PerfReportTool
 			}
 			if (fileType == FileType.TextCsv)
 			{
+				// TODO: Investigate moving this to the PrepareCsvData function so it can be done in parallel like with binary CSVs
+				// Not sure why it's not like that currently - maybe mem usage was the issue, or it was faster for some reason
 				finalCsv = CsvStats.ReadCSVFromLines(textCsvLines, null);
+				AddDerivedMetadataToCsv(finalCsv);
 
 				// If we already have metadata then use it. This avoids the need to remap it multiple times
 				if ( metadata != null )
 				{
 					finalCsv.metaData = metadata;
 				}
-
 				return finalCsv;
 			}
 			return null;
 		}
 
-
 		public bool DoesMetadataMatchQuery(QueryExpression metadataQuery)
 		{
 			if (metadataQuery == null)
@@ -226,6 +242,7 @@ namespace PerfReportTool
 		public SummaryTableRowData cachedSummaryTableRowData;
 		public ReportTypeInfo reportTypeInfo;
 		public XmlVariableMappings xmlVariableMappings;
+		public DerivedCsvStatDefinition[] derivedCsvStatDefinitions;
 		FileType fileType;
 	};
 
@@ -269,6 +286,7 @@ namespace PerfReportTool
 			bRemoveDuplicates = inRemoveDuplicates;
 			bRequireMetadata = inRequireMetadata;
 			bListFilesMode = inListFilesMode;
+			derivedCsvStatDefinitions = reportXml.derivedCsvStatDefinitions.Values.ToArray();
 
 			csvIdToFilename = new Dictionary<string, string>();
 			outCsvQueue = new BlockingCollection<CsvFileInfo>(lookaheadCount);
@@ -301,7 +319,7 @@ namespace PerfReportTool
 			if (precacheThreads == null)
 			{
 				CsvFileInfo fileInfo = csvFileInfos[readIndex];
-				file = new CachedCsvFile(fileInfo.filename, useCacheFiles, derivedMetadataMappings);
+				file = new CachedCsvFile(fileInfo.filename, useCacheFiles, derivedMetadataMappings, derivedCsvStatDefinitions);
 				if (!file.DoesMetadataMatchQuery(metadataQuery)) // TODO do we need to check this here and in the thread?
 				{
 					return null;
@@ -348,7 +366,7 @@ namespace PerfReportTool
 				CsvFileInfo fileInfo = csvFileInfos[threadWriteIndex];
 
 				// Process the file
-				CachedCsvFile file = new CachedCsvFile(fileInfo.filename, useCacheFiles, derivedMetadataMappings);
+				CachedCsvFile file = new CachedCsvFile(fileInfo.filename, useCacheFiles, derivedMetadataMappings, derivedCsvStatDefinitions);
 
 				bool bProcessThisFile = true;
 				if (bRequireMetadata && file.metadata == null)
@@ -462,6 +480,7 @@ namespace PerfReportTool
 		ReportXML reportXml;
 		ReportTypeParams reportTypeParams;
 		DerivedMetadataMappings derivedMetadataMappings;
+		DerivedCsvStatDefinition[] derivedCsvStatDefinitions;
 
 		Dictionary<string, string> csvIdToFilename;
 		public int duplicateCount = 0;
diff --git a/Engine/Source/Programs/CSVTools/PerfReportTool/PerfReportTool.cs b/Engine/Source/Programs/CSVTools/PerfReportTool/PerfReportTool.cs
--- a/Engine/Source/Programs/CSVTools/PerfReportTool/PerfReportTool.cs
+++ b/Engine/Source/Programs/CSVTools/PerfReportTool/PerfReportTool.cs
@@ -23,7 +23,7 @@ namespace PerfReportTool
     class Version
     {
 		// Format: Major.Minor.Bugfix
-        private static string VersionString = "4.240.1";
+        private static string VersionString = "4.241.0";
 
         public static string Get() { return VersionString; }
     };
@@ -1026,18 +1026,6 @@ namespace PerfReportTool
 			{
 				perfLog.LogTiming("    ReadCsvStats");
 			}
-
-			// Compute derived CSV stats
-			foreach (DerivedCsvStatDefinition derivedStat in reportXML.derivedCsvStatDefinitions.Values)
-			{
-				csvStats.AddDerivedStatFromExpression(derivedStat.getExpressionStr());
-			}
-			if (perfLog != null)
-			{
-				perfLog.LogTiming("    Compute Derived CSV stats");
-			}
-
-
 			if (bStripStatsByEvents)
 			{
 				CsvStats strippedCsvStats = StripCsvStatsByEvents(csvStatsUnstripped, out numFramesStripped);
@@ -1086,6 +1074,31 @@ namespace PerfReportTool
 			}
 		}
 
+		string ExpandGraphStatListDerivedStats(string originalGraphStatNamesStr, CsvStats csvStats)
+		{
+			string[] statNames = originalGraphStatNamesStr.ToLowerInvariant().Split(',');
+			List<string> newStatNames = new List<string>();
+			bool bHasDerivedStats = false;
+			for (int i = 0; i < statNames.Length; i++)
+			{
+				List<StatSamples> matchingStats = csvStats.GetStatsMatchingString(statNames[i]);
+				foreach (StatSamples stat in matchingStats)
+				{
+					DerivedCsvStatDefinition derivedStatDefinition = stat.bIsDerived ? reportXML.GetDerivedCsvStatDefinition(stat.Name) : null;
+					if (derivedStatDefinition != null)
+					{
+						newStatNames.Add(derivedStatDefinition.getExpressionStr());
+						bHasDerivedStats = true;
+					}
+					else
+					{
+						newStatNames.Add(stat.Name);
+					}
+				}
+			}
+			return bHasDerivedStats ? String.Join(',', newStatNames) : originalGraphStatNamesStr;
+		}
+
 		void GenerateReport(CachedCsvFile csvFile, string outputDir, bool bBulkMode, SummaryTableRowData rowData, bool bBatchedGraphs, bool writeDetailedReport, bool bReadCsvStats, ReportTypeInfo reportTypeInfo, string csvDir)
 		{
 			PerfLog perfLog = new PerfLog(GetBoolArg("perfLog"));
@@ -1128,6 +1141,7 @@ namespace PerfReportTool
 
 			bool bCsvToSvgProcesses = GetBoolArg("csvToSvgProcesses");
 			bool bCsvToSvgMultiThreaded = !GetBoolArg("csvToSvgSequential");
+			bool useEmbeddedGraphUrl = GetBoolArg("useEmbeddedGraphUrl");
 
 			int minWorker, minIOC;
 			// Get the current settings.
@@ -1140,16 +1154,15 @@ namespace PerfReportTool
 			string responseFilename = null;
 			List<Process> csvToSvgProcesses = new List<Process>();
 			List<Task> csvToSvgTasks = new List<Task>();
+
 			if (writeDetailedReport)
 			{
 				GraphGenerator graphGenerator = null;
-				if (!bCsvToSvgProcesses)
+				if (!bCsvToSvgProcesses && !useEmbeddedGraphUrl)
 				{
 					graphGenerator = new GraphGenerator(csvFile.GetFinalCsv(), csvFile.filename);
 				}
 
-				bool useEmbeddedGraphUrl = GetBoolArg("useEmbeddedGraphUrl");
-
 				// Generate all the graphs asyncronously
 				foreach (ReportGraph graph in reportTypeInfo.graphs)
 				{
@@ -1158,31 +1171,16 @@ namespace PerfReportTool
 						throw new Exception("Graph " + graph.title + " has no <statString> element");
 					}
 
-					// Replace any referenced derived stats with the full expression so that CSVToSVG can expand them
-					// TODO: Technically this should only be necessary if CSVToSVG is running as a separate process
-					if (reportXML.derivedCsvStatDefinitions.Count > 0)
+					bool bShowGraph = graph.ShouldShow(csvFile.finalCsv ?? csvFile.dummyCsvStats);
+					if (bShowGraph)
 					{
-						string[] stats = graph.settings.statString.value.ToLowerInvariant().Split(',');
-						bool bHasDerivedStats = false;
-						for (int i = 0; i < stats.Length; i++)
-						{
-							DerivedCsvStatDefinition derivedStatDefinition = reportXML.GetDerivedCsvStatDefinition(stats[i].Trim());
-							if (derivedStatDefinition != null)
-							{
-								stats[i] = derivedStatDefinition.getExpressionStr();
-								bHasDerivedStats = true;
-							}
-						}
-						if (bHasDerivedStats) 
+						// If we're not generating graphs in this process then expand any referenced derived stats with the full expression so that CsvToSvg can compute them itself
+						// (We've already computed them in this tool, but CsvToSvg only has the original Csv file)
+						if (graphGenerator == null && reportXML.derivedCsvStatDefinitions.Count > 0)
 						{
-							graph.settings.statString.value = String.Join(',', stats);
+							graph.settings.statString.value = ExpandGraphStatListDerivedStats(graph.settings.statString.value, csvFile.finalCsv ?? csvFile.dummyCsvStats);
 						}
-					}
 
-
-					bool bShowGraph = graph.ShouldShow(csvFile.finalCsv ?? csvFile.dummyCsvStats);
-					if (bShowGraph)
-					{
 						if (useEmbeddedGraphUrl)
 						{
 							csvSvgInfoList.Add(CsvSvgInfo.CreateEmbeddedUrlGraphInfo(graph));
diff --git a/Engine/Source/Programs/CSVTools/PerfReportTool/Summaries/PeakSummary.cs b/Engine/Source/Programs/CSVTools/PerfReportTool/Summaries/PeakSummary.cs
--- a/Engine/Source/Programs/CSVTools/PerfReportTool/Summaries/PeakSummary.cs
+++ b/Engine/Source/Programs/CSVTools/PerfReportTool/Summaries/PeakSummary.cs
@@ -41,18 +41,18 @@ namespace PerfSummaries
 			}
 			public bool StatMatchesSection(CsvStats csvStats, string statName)
 			{
-				if (statNameFilterDict == null)
+				if (statNameFilterSet == null)
 				{
-					statNameFilterDict = csvStats.GetStatNamesMatchingStringList_Dict(statNamesFilter);
+					statNameFilterSet = csvStats.GetStatNamesMatchingStringList_Set(statNamesFilter);
 				}
-				return statNameFilterDict.ContainsKey(statName);
+				return statNameFilterSet.Contains(statName);
 			}
 			public void AddStat(PeakStatInfo statInfo)
 			{
 				stats.Add(statInfo);
 			}
 			string[] statNamesFilter;
-			Dictionary<string, bool> statNameFilterDict;
+			HashSet<string> statNameFilterSet;
 
 			public List<PeakStatInfo> stats = new List<PeakStatInfo>();
 			public string title;
