CsvToSvg 3.66: Fix the stat title/filename when derived stats are passed as the first param

diff --git a/Engine/Source/Programs/CSVTools/CSVToSVG/CsvToSvgProgram.cs b/Engine/Source/Programs/CSVTools/CSVToSVG/CsvToSvgProgram.cs
--- a/Engine/Source/Programs/CSVTools/CSVToSVG/CsvToSvgProgram.cs
+++ b/Engine/Source/Programs/CSVTools/CSVToSVG/CsvToSvgProgram.cs
@@ -234,7 +234,14 @@ namespace CSVTools
 					{
 						svgFilename += "_";
 					}
-					svgFilename += graphParams.statNames[0].Replace('*', 'X').Replace('/', '_');
+					string statName = graphParams.statNames[0];
+					// If this is a stat expression and we've defined the name via the = operator then use that
+					int idx = statName.IndexOf('=');
+					if (idx > 0)
+					{
+						statName = statName[..idx];
+					}
+					svgFilename += statName.Replace('*', 'X').Replace('/', '_').Replace("(", "").Replace(")", "");
 				}
 				svgFilename += ".svg";
 			}
diff --git a/Engine/Source/Programs/CSVTools/CsvToSvgLib/CsvToSvgLib.cs b/Engine/Source/Programs/CSVTools/CsvToSvgLib/CsvToSvgLib.cs
--- a/Engine/Source/Programs/CSVTools/CsvToSvgLib/CsvToSvgLib.cs
+++ b/Engine/Source/Programs/CSVTools/CsvToSvgLib/CsvToSvgLib.cs
@@ -15,7 +15,7 @@ namespace CSVTools
 
 	public class CsvToSvgLibVersion
 	{
-		private static string VersionString = "3.65";
+		private static string VersionString = "3.66";
 
 		public static string Get() { return VersionString; }
 	};
@@ -646,11 +646,20 @@ namespace CSVTools
 			// Get the title
 			if (graphTitle.Length == 0 && statNames.Count == 1 && csvStatsList.Count > 0 && !statNames[0].Contains("*"))
 			{
+				// Use the stat name if we can find it, otherwise just take the param as-is
 				StatSamples stat = csvStatsList[0].GetStat(statNames[0]);
 				if (stat == null)
 				{
-					Console.Out.WriteLine("Warning: Could not find stat {0}", statNames[0]);
-					graphTitle = string.Format("UnknownStat {0}", statNames[0]);
+					int idx = statNames[0].IndexOf('=');
+					if (idx > 0)
+					{
+						// If this is a stat expression and we've defined the name via the = operator then use that
+						graphTitle = statNames[0][..idx].Replace("(","").Replace(")","");
+					}
+					else
+					{
+						graphTitle = statNames[0];
+					}					
 				}
 				else
 				{
