PerfReportTool/CsvToSVG/CsvStats: Derived stats improvements:
- Add eventcount() function to read per-frame event counts in expressions based on a search string, e.g: syncLoads=eventcount('sync load*')
- Add a meta() function to use numeric metadata values in expressions
- Rename the add() function to addstats(), because it can only be used for stats
- Make event parsing more robust

diff --git a/Engine/Source/Programs/CSVTools/CsvStats/CsvStats.cs b/Engine/Source/Programs/CSVTools/CsvStats/CsvStats.cs
--- a/Engine/Source/Programs/CSVTools/CsvStats/CsvStats.cs
+++ b/Engine/Source/Programs/CSVTools/CsvStats/CsvStats.cs
@@ -12,7 +12,6 @@ using System.Collections;
 using System.IO.Compression;
 using System.Security.Cryptography;
 using System.Data;
-using System.Linq.Expressions;
 
 namespace CSVStats
 {
@@ -504,6 +503,15 @@ namespace CSVStats
 			return null;
 		}
 
+		public string GetMetadataValue(string key)
+		{
+			if (metaData == null)
+			{
+				return null;
+			}
+			return metaData.GetValue(key, null);
+		}
+
 		private string ExpandCustomFunctions(string expression)
 		{
 			char[] operators = { '+', '-', '*', '/' };
@@ -539,6 +547,206 @@ namespace CSVStats
 			return expression;
 		}
 
+		// Matches 
+		string ReadQuotedString(string str)
+		{
+			if (str.Length < 2 || str[0] != '\'')
+			{
+				return null;
+			}
+			int endIdx = str.IndexOf('\'',1);
+			if (endIdx == -1)
+			{
+				return null;
+			}
+			return str[1..endIdx];
+		}
+		string ReadBracketString(string str)
+		{
+			if (str.Length < 2 || str[0] != '(')
+			{
+				return null;
+			}
+			string strOut = "";
+			int bracketDepth = 1;
+			bool bInQuotes = false;
+			for (int i=1; i<str.Length; i++)
+			{
+				char c = str[i];
+				if (c=='\'')
+				{
+					bInQuotes = !bInQuotes;
+				}
+				if (!bInQuotes)
+				{
+					if (c == '(')
+					{
+						bracketDepth++;
+					}
+					else if (c == ')')
+					{
+						bracketDepth--;
+						if (bracketDepth == 0)
+						{
+							return strOut;
+						}
+					}
+				}
+				strOut += c;
+			}
+			return null;
+		}
+
+		enum BuiltinFunctionType { addstats, meta, eventcount, COUNT };
+		string ReadNextToken(string expression, ref int index, Dictionary<string, StatSamples> statsOut, int tokenIndex)
+		{
+			// Skip leading whitespace
+			while (index < expression.Length && Char.IsWhiteSpace(expression[index]))
+			{
+				index++;
+			}
+
+			if (index >= expression.Length)
+			{
+				// We're done
+				return null;
+			}
+
+			// Match delimiters
+			char[] delimiters = { '+', '-', '*', '(', ')', '/' };
+			if (delimiters.Contains(expression[index]))
+			{
+				index++;
+				return expression[index-1].ToString();
+			}
+
+			string remaining = expression[index..];
+			int oldIndex = index;
+
+			// Match our builtin functions. These are all of the form NAME('PARAM')
+			for ( int i=0; i<(int)BuiltinFunctionType.COUNT; i++)
+			{
+				BuiltinFunctionType funcType = (BuiltinFunctionType)i;
+				string funcName = funcType.ToString();
+				if (remaining.Length > (funcName.Length + 4) && 
+					remaining.StartsWith(funcName+"("))
+				{
+					string paramString = ReadBracketString(remaining[(funcName.Length)..]);
+					if (paramString == null)
+					{
+						continue;
+					}
+					// TODO: process the function here
+					index += funcName.Length + paramString.Length + 2;
+
+					if (funcType == BuiltinFunctionType.meta)
+					{
+						string key = ReadQuotedString(paramString);
+						string val = key == null ? null : GetMetadataValue(key);
+						float valueNumeric;
+						if (val==null || !float.TryParse(val, out valueNumeric))
+						{
+							throw new Exception("Warning: failed to read metadata: " + expression[oldIndex..index]);
+						}
+						return val.ToString();
+					}
+					else if (funcType == BuiltinFunctionType.eventcount)
+					{
+						string eventStr = ReadQuotedString(paramString);
+						if (eventStr == null)
+						{
+							throw new Exception("Warning: eventstr requires a quote '' arg: " + expression[oldIndex..index]);
+						}
+						string statName = "__STAT__EVENTCOUNT__" + tokenIndex.ToString();
+						statsOut[statName] = MakeEventCountStat(statName, eventStr);
+						return statName;
+					}
+					else if (funcType == BuiltinFunctionType.addstats)
+					{
+						// Just remove any quotes. They're not needed since we're just literally matching stat names
+						string statsStr = paramString.Replace("'", "");
+						string[] statNames = statsStr.Split(",");
+						List<StatSamples> foundStats = GetStatsMatchingStringList(statNames);
+						if (foundStats.Count == 0)
+						{
+							// Fail gracefully if the stats weren't found. The sum of no stats is zero
+							return "0";
+						}
+						// Sanitize the stat names and add the stats to statsOut
+						List<string> safeStatNames = new List<string>();
+						foreach (StatSamples stat in foundStats)
+						{
+							string safeName = stat.Name.Replace("/", "__SLASH__");
+							safeStatNames.Add(safeName);
+							statsOut[safeName] = stat;
+						}
+						// Return the stats added together
+						return "("+String.Join("+", safeStatNames)+")";
+					}
+				}
+			}
+
+			// Match stats
+			if (remaining.StartsWith('\''))
+			{
+				// Read the literal stat name in quotes. This is needed for stat names that include slashes
+				string statName = ReadQuotedString(remaining);
+				index += statName.Length + 2;
+				StatSamples stat = GetStat(statName.Trim());
+				if (stat == null)
+				{
+					throw new Exception("Stat not found: " + statName);
+				}
+				string safeName = stat.Name.Replace("/", "__SLASH__");
+				statsOut[safeName] = stat;
+				return safeName;
+			}
+			// Read up to the next delimiter as the stat name
+			int endIdx = remaining.Length;
+			for (int i=0; i<remaining.Length; i++)
+			{
+				if (delimiters.Contains(remaining[i]))
+				{
+					endIdx = i;
+					break;
+				}
+			}
+			string nextToken = remaining[0..endIdx];
+			index += nextToken.Length;
+			
+			// Check if this is a stat. If not, just return it unmodified
+			StatSamples foundStat = GetStat(nextToken.Trim());
+			if ( foundStat != null )
+			{
+				statsOut[foundStat.Name] = foundStat;
+				return foundStat.Name;
+			}
+			return expression[oldIndex..index];
+		}
+
+		string PreProcessExpression(string expression, Dictionary<string,StatSamples> statsOut)
+		{
+			List<string> tokens = new List<string>();
+			expression = expression.Trim();
+			if (expression.Length == 0)
+			{
+				return null;
+			}
+
+			int index = 0;
+			string nextToken;
+			do 
+			{
+				StatSamples stat = null;
+				nextToken = ReadNextToken(expression, ref index, statsOut, tokens.Count);
+				if (nextToken != null)
+				{
+					tokens.Add(nextToken);
+				}
+			} while(nextToken != null);
+			return String.Join("", tokens);
+		}
+
 		public StatSamples MakeDerivedStatFromExpression(string expression)
 		{
 			if (expression == "")
@@ -546,8 +754,7 @@ namespace CSVStats
 				return null;
 			}
 
-			// Remove whitespace
-			expression = expression.Replace(" ", "");
+			expression = expression.Trim();
 
 			// Remove outer brackets if they exist
 			if (expression[0] == '(' && expression[^1]==')' )
@@ -578,85 +785,15 @@ namespace CSVStats
 			expression = expression.ToLower();
 
 			Dictionary<string, StatSamples> referencedStats = new Dictionary<string, StatSamples>();
-
-			// Find stats in single quotes and generate sanitized names for them (ie replace / with _SLASH_)
-			// This works for individual stats or for multiple stats separated by commas
-			if (expression.Contains('\''))
+			string finalExpression;
+			try
 			{
-				string newExpression = "";
-				bool bInQuotes = false;
-				string token = "";
-				for (int i=0; i<expression.Length; i++) 
-				{
-					char c = expression[i];
-					if (c == '\'') 
-					{
-						bInQuotes = !bInQuotes;
-						if ( bInQuotes )
-						{
-							token = "";
-						}
-						else
-						{
-							if (token.Length>0)
-							{
-								// Find the matching stat. If a wildcard is specified, multiple stats may be returned
-								List<StatSamples> matchingStats = GetStatsMatchingStringList(token.Split(','));
-								if (matchingStats.Count == 0)
-								{
-									Console.WriteLine("Warning: stat not found: " + token);
-									return null;
-								}
-								bool bFirst = true;
-								foreach (StatSamples stat in matchingStats)
-								{
-									string sanitizedStatName = stat.Name.ToLower().Replace("/", "_SLASH_");
-									if (!bFirst)
-									{
-										newExpression += ',';
-									}
-									newExpression += sanitizedStatName;
-									referencedStats[sanitizedStatName] = stat;
-									bFirst = false;
-								}
-							}
-						}
-					}
-					else
-					{
-						if (bInQuotes)
-						{
-							token += c;
-						}
-						else
-						{
-							newExpression += c;
-						}
-					}
-				}
-				expression = newExpression;
+				finalExpression = PreProcessExpression(expression, referencedStats);
 			}
-
-			// Handle special functions we might want to expand
-			expression = ExpandCustomFunctions(expression);
-
-			// Find any other refererenced stats
-			char[] delimiters = { '+', '-', '*', '(', ')', '/' };
-			string[] tokens = expression.Split(delimiters);
-			foreach (string token in tokens)
+			catch (Exception ex)
 			{
-				if (token == "")
-				{
-					continue;
-				}
-				if (!referencedStats.ContainsKey(token))
-				{
-					StatSamples stat = GetStat(token);
-					if (stat != null)
-					{
-						referencedStats.Add(token, stat);
-					}
-				}
+				Console.WriteLine("Error: derived stat expression \"" + expression + "\"could not be computed. Preprocessing error:" + ex.ToString());
+				return null;
 			}
 
 			// Construct the datatable with any referenced stats
@@ -681,7 +818,7 @@ namespace CSVStats
 			DataColumn resultColumn = new DataColumn("__RESULT__", typeof(float));
 			try
 			{
-				resultColumn.Expression = expression;
+				resultColumn.Expression = finalExpression;
 				dataTable.Columns.Add(resultColumn);
 				dataTable.AcceptChanges();
 			}
@@ -699,7 +836,6 @@ namespace CSVStats
 				statSamples.samples[i] = (float)row[resultColumnIndex];
 			}
 			statSamples.ComputeAverageAndTotal();
-
 			return statSamples;
 		}
 
@@ -765,6 +901,34 @@ namespace CSVStats
 			return listOut;
 		}
 
+		public StatSamples MakeEventCountStat(string statName, string eventStr)
+		{
+			StatSamples newStat = new StatSamples(statName, SampleCount);
+			List<CsvEvent> matchingEvents = GetEventsMatchingString(eventStr);
+			foreach (CsvEvent ev in matchingEvents)
+			{
+				if (ev.Frame < SampleCount)
+				{
+					newStat.samples[ev.Frame]++;
+				}
+			}
+			newStat.ComputeAverageAndTotal();
+			return newStat;
+		}
+
+		public List<CsvEvent> GetEventsMatchingString(string eventStr)
+		{
+			List<CsvEvent> eventsOut = new List<CsvEvent>();
+			foreach (CsvEvent ev in Events)
+			{
+				if (CsvStats.DoesSearchStringMatch(ev.Name, eventStr))
+				{
+					eventsOut.Add(ev);
+				}
+			}
+			return eventsOut;
+		}
+
 
 		public void AddStat(StatSamples stat)
 		{
