CsvStats: Derived stats: make quotes for builtin functions optional since they're not required and this makes them less error prone

diff --git a/Engine/Source/Programs/CSVTools/CsvStats/CsvStats.cs b/Engine/Source/Programs/CSVTools/CsvStats/CsvStats.cs
--- a/Engine/Source/Programs/CSVTools/CsvStats/CsvStats.cs
+++ b/Engine/Source/Programs/CSVTools/CsvStats/CsvStats.cs
@@ -636,13 +636,14 @@ namespace CSVStats
 					{
 						continue;
 					}
-					// TODO: process the function here
 					index += funcName.Length + paramString.Length + 2;
 
 					if (funcType == BuiltinFunctionType.meta)
 					{
+						// Quotes are optional. They're supported for consistency with stats
 						string key = ReadQuotedString(paramString);
-						string val = key == null ? null : GetMetadataValue(key);
+						key = (key == null) ? paramString : key;
+						string val = GetMetadataValue(key);
 						float valueNumeric;
 						if (val==null || !float.TryParse(val, out valueNumeric))
 						{
@@ -652,7 +653,9 @@ namespace CSVStats
 					}
 					else if (funcType == BuiltinFunctionType.eventcount)
 					{
+						// Quotes are optional. They're supported for consistency with stats
 						string eventStr = ReadQuotedString(paramString);
+						eventStr = (eventStr == null) ? paramString : eventStr;
 						if (eventStr == null)
 						{
 							throw new Exception("Warning: eventstr requires a quote '' arg: " + expression[oldIndex..index]);
@@ -663,7 +666,7 @@ namespace CSVStats
 					}
 					else if (funcType == BuiltinFunctionType.addstats)
 					{
-						// Just remove any quotes. They're not needed since we're just literally matching stat names
+						// Quotes are optional. They're supported for consistency with other stats
 						string statsStr = paramString.Replace("'", "");
 						string[] statNames = statsStr.Split(",");
 						List<StatSamples> foundStats = GetStatsMatchingStringList(statNames);
@@ -737,7 +740,6 @@ namespace CSVStats
 			string nextToken;
 			do 
 			{
-				StatSamples stat = null;
 				nextToken = ReadNextToken(expression, ref index, statsOut, tokens.Count);
 				if (nextToken != null)
 				{
