Make CsvToSvg highlightEventRegions param consistent with other list params and use the semicolon separator

diff --git a/Engine/Source/Programs/CSVTools/CSVToSVG/CsvToSvgProgram.cs b/Engine/Source/Programs/CSVTools/CSVToSVG/CsvToSvgProgram.cs
--- a/Engine/Source/Programs/CSVTools/CSVToSVG/CsvToSvgProgram.cs
+++ b/Engine/Source/Programs/CSVTools/CSVToSVG/CsvToSvgProgram.cs
@@ -66,7 +66,7 @@ namespace CSVTools
 			"       -showAllEventNames <1|0> - don't filter out event names if they're too dense\n" +
 			"       -hideStatPrefix <list>\n" +
 			"       -hierarchySeparator <character>\n" +
-			"       -highlightEventRegions <startEventName,endEventName>\n" +
+			"       -highlightEventRegions <startEventName0;endEventName0;...>\n" +
 			"       -ignoreStats <list> (can include wildcards. Separate states with stat1;stat2;etc)\n" +
 			"       -interactive\n" +
 			"       -legend <list> \n" +
@@ -343,7 +343,7 @@ namespace CSVTools
 			{
 				graphParams.showEventNameTextMode = ShowEventTextMode.ShowAll;
 			}
-			graphParams.highlightEventRegions = GetListArg("highlightEventRegions", ',');
+			graphParams.highlightEventRegions = GetListArg("highlightEventRegions");
 
 			// Start/end event
 			graphParams.startEvent = GetArg("startEvent", null);
diff --git a/Engine/Source/Programs/CSVTools/CsvToSvgLib/CsvToSvgLib.cs b/Engine/Source/Programs/CSVTools/CsvToSvgLib/CsvToSvgLib.cs
--- a/Engine/Source/Programs/CSVTools/CsvToSvgLib/CsvToSvgLib.cs
+++ b/Engine/Source/Programs/CSVTools/CsvToSvgLib/CsvToSvgLib.cs
@@ -15,7 +15,7 @@ namespace CSVTools
 
 	public class CsvToSvgLibVersion
 	{
-		private static string VersionString = "3.66";
+		private static string VersionString = "3.67";
 
 		public static string Get() { return VersionString; }
 	};
diff --git a/Engine/Source/Programs/CSVTools/PerfReportTool/PerfReportTool.cs b/Engine/Source/Programs/CSVTools/PerfReportTool/PerfReportTool.cs
--- a/Engine/Source/Programs/CSVTools/PerfReportTool/PerfReportTool.cs
+++ b/Engine/Source/Programs/CSVTools/PerfReportTool/PerfReportTool.cs
@@ -23,7 +23,7 @@ namespace PerfReportTool
     class Version
     {
 		// Format: Major.Minor.Bugfix
-        private static string VersionString = "4.241.0";
+        private static string VersionString = "4.241.1";
 
         public static string Get() { return VersionString; }
     };
@@ -1892,11 +1892,11 @@ namespace PerfReportTool
 					{
 						if (i > 0)
 						{
-							highlightEventRegions += ",";
+							highlightEventRegions += ";";
 						}
 						string endEvent = (eventsToStrip[i].endName == null) ? "{null}" : eventsToStrip[i].endName;
 						string beginEvent = (eventsToStrip[i].beginName == null) ? "{null}" : eventsToStrip[i].beginName;
-						highlightEventRegions += beginEvent + "," + endEvent;
+						highlightEventRegions += beginEvent + ";" + endEvent;
 					}
 					highlightEventRegions += "\"";
 				}
