Added CSVToTelemetry Tool
Tool takes a CSVFile as input and emits a Telemetry event for each row in the file.
Tool asumes that the first row always contains the event payload columns names.
Added CSVToTelemetry upload to Generate IoChunks CSV in Nightly Build script.

diff --git a/Engine/Source/Programs/CSVTools/CSVToTelemetry/CSVToTelemetry.Build.cs b/Engine/Source/Programs/CSVTools/CSVToTelemetry/CSVToTelemetry.Build.cs
new file mode 100644
--- /dev/null
+++ b/Engine/Source/Programs/CSVTools/CSVToTelemetry/CSVToTelemetry.Build.cs
@@ -0,0 +1,25 @@
+// Copyright Epic Games, Inc. All Rights Reserved.
+
+using UnrealBuildTool;
+
+public class CSVToTelemetry : ModuleRules
+{
+	public CSVToTelemetry(ReadOnlyTargetRules Target) : base(Target)
+	{
+		PublicIncludePathModuleNames.Add("Launch");
+
+		PrivateDependencyModuleNames.AddRange(new string[] {
+			"AnalyticsET",
+			"ApplicationCore",
+			"AssetRegistry",
+			"BuildSettings",
+			"Core",
+			"CoreUObject",
+			"Horde",
+			"HTTP",
+			"Projects",
+			"RSA",
+			"TargetPlatform"
+		});
+	}
+}
diff --git a/Engine/Source/Programs/CSVTools/CSVToTelemetry/CSVToTelemetry.Target.cs b/Engine/Source/Programs/CSVTools/CSVToTelemetry/CSVToTelemetry.Target.cs
new file mode 100644
--- /dev/null
+++ b/Engine/Source/Programs/CSVTools/CSVToTelemetry/CSVToTelemetry.Target.cs
@@ -0,0 +1,47 @@
+// Copyright Epic Games, Inc. All Rights Reserved.
+
+using UnrealBuildTool;
+using System.Collections.Generic;
+
+public class CSVToTelemetryTarget : TargetRules
+{
+	public CSVToTelemetryTarget(TargetInfo Target) : base(Target)
+	{
+		Type = TargetType.Program;
+		LinkType = TargetLinkType.Modular;
+		IncludeOrderVersion = EngineIncludeOrderVersion.Latest;
+		LaunchModuleName = "CSVToTelemetry";
+
+		bUseXGEController					= false;
+		bCompileFreeType					= false;
+		bLoggingToMemoryEnabled				= false;
+		bUseLoggingInShipping				= true;
+		bCompileWithAccessibilitySupport	= false;
+		bCompileWithPluginSupport			= true;
+		bIncludePluginsForTargetPlatforms	= true;
+		bWithServerCode						= false;
+		bCompileNavmeshClusterLinks			= false;
+		bCompileNavmeshSegmentLinks			= false;
+		bCompileRecast						= false;
+		bCompileICU 						= false;
+		bWithLiveCoding						= false;
+		bBuildDeveloperTools				= false;
+		bBuildWithEditorOnlyData			= true;
+		bCompileAgainstEngine				= false;
+		bCompileAgainstCoreUObject			= true;
+		bCompileAgainstApplicationCore		= true;
+		bUsesSlate							= false;
+		bIsBuildingConsoleApplication		= true;
+		bForceBuildTargetPlatforms			= true;
+		bNeedsExtraShaderFormats			= true;
+		bAreTargetSDKVersionsRelevantOverride = false;
+
+		WindowsPlatform.bUseBundledDbgHelp = false;
+		WindowsPlatform.bPixProfilingEnabled = false;
+		WindowsPlatform.bMergeIdenticalCOMDATs = true;
+		WindowsPlatform.bStripUnreferencedSymbols = true;
+
+		bEnableTrace = true;
+		GlobalDefinitions.Add("UE_EXTERNAL_PROFILING_ENABLED=0");
+	}
+}
diff --git a/Engine/Source/Programs/CSVTools/CSVToTelemetry/Private/CSVToTelemetry.cpp b/Engine/Source/Programs/CSVTools/CSVToTelemetry/Private/CSVToTelemetry.cpp
new file mode 100644
--- /dev/null
+++ b/Engine/Source/Programs/CSVTools/CSVToTelemetry/Private/CSVToTelemetry.cpp
@@ -0,0 +1,261 @@
+// Copyright Epic Games, Inc. All Rights Reserved.
+
+#include "CSVToTelemetry.h"
+#include "RequiredProgramMainCPPInclude.h"
+#include "Analytics.h"
+#include "AnalyticsET.h"
+#include "BuildSettings.h"
+#include "Horde.h"
+#include "HttpModule.h"
+#include "HttpManager.h"
+#include "IAnalyticsProviderET.h"
+#include "Logging/LogMacros.h"
+#include "ProjectUtilities.h"
+
+DEFINE_LOG_CATEGORY_STATIC(LogCSVToTelemetry, Log, All);
+IMPLEMENT_APPLICATION(CSVToTelemetry, "CSVToTelemetry");
+
+class FTelemetry
+{
+public:
+
+	bool StartSession()
+	{
+		FString APIServer;
+
+		if (AnalyticsProvider.IsValid())
+		{
+			// Already running the telemetry
+			return false;
+		}
+		
+		// We must have supplied at least the APIServer URL to send any telemetry
+		if (FParse::Value(FCommandLine::Get(), TEXT("server="), APIServer)==false)
+		{
+			UE_LOG(LogCSVToTelemetry, Error, TEXT("Must provide a valid server address with -server=[serverURL]"));
+			return false;
+		}
+
+		// Allow override of API Endpoint
+		FString APIEndpoint = TEXT("datarouter/api/v1/public/data");
+		FParse::Value(FCommandLine::Get(), TEXT("endpoint="), APIEndpoint);
+
+		// Allow override of API Endpoint
+		FString APIKey = TEXT("StudioAnalytics.Dev");
+		FParse::Value(FCommandLine::Get(), TEXT("apikey="), APIKey);
+
+		FAnalyticsET::Config Config;
+
+		Config.APIServerET = MoveTemp(APIServer);
+		Config.APIEndpointET = MoveTemp(APIEndpoint);
+		Config.APIKeyET = MoveTemp(APIKey);
+		Config.AppEnvironment = TEXT("datacollector-binary");
+		Config.AppVersionET = FEngineVersion::Current().ToString();
+
+		AnalyticsProvider = FAnalyticsET::Get().CreateAnalyticsProvider(Config);
+
+		if (AnalyticsProvider.IsValid())
+		{
+			// Prepare the default attributes appended to all events
+			TArray<FAnalyticsEventAttribute> DefaultAttributes = MakeAnalyticsEventAttributeArray(					
+				TEXT("Horde_TemplateName"), FHorde::GetTemplateName(),
+				TEXT("Horde_StepName"), FHorde::GetStepName(),
+				TEXT("Build_Changelist"), BuildSettings::GetCurrentChangelist(),
+				TEXT("Build_BranchName"), FApp::GetBranchName().ToLower());
+
+			uint32 SchemaVersion(0);
+
+			if (FParse::Value(FCommandLine::Get(), TEXT("schema="), SchemaVersion))
+			{
+				// Optional schema version for this event
+				DefaultAttributes.Emplace(TEXT("SchemaVersion"), SchemaVersion);
+			}
+
+			// Set the default attributes
+			AnalyticsProvider->SetDefaultEventAttributes(MoveTemp(DefaultAttributes));
+
+			// Start the provider session
+			AnalyticsProvider->StartSession();
+
+			return true;
+		}
+		
+		return false;
+	}
+
+	void RecordEvent(const FString& EventName, const TArray<FAnalyticsEventAttribute>& Attributes)
+	{
+		if (AnalyticsProvider.IsValid())
+		{
+			AnalyticsProvider->RecordEvent(EventName, Attributes);
+		}
+	}
+
+	void EndSession()
+	{
+		if (AnalyticsProvider.IsValid())
+		{		
+			// End the session
+			AnalyticsProvider->FlushEvents();
+			AnalyticsProvider->EndSession();
+			AnalyticsProvider.Reset();
+
+			UE_LOG(LogCSVToTelemetry, Display, TEXT("Waiting for Telemetry upload to complete"));
+
+			double StartTime = StartTime = FPlatformTime::Seconds();
+			// It is quite likely that one of the analytics providers is sending data via the FHttpManager so we
+			// need to flush that as well to make sure that the message is sent. 
+			if (FHttpModule* HttpModule = FModuleManager::GetModulePtr<FHttpModule>("HTTP"))
+			{
+				FHttpManager& HttpManager = HttpModule->GetHttpManager();
+				HttpManager.Flush(EHttpFlushReason::FullFlush);
+			}
+
+			UE_LOG(LogCSVToTelemetry, Display, TEXT("Telemetry upload finished in %f seconds"), FPlatformTime::Seconds() - StartTime);
+		}
+	}
+
+private:
+
+	TSharedPtr<IAnalyticsProviderET>	AnalyticsProvider;
+	float								TimeOut = 1200.0f;
+};
+
+INT32_MAIN_INT32_ARGC_TCHAR_ARGV()
+{
+	FTaskTagScope Scope(ETaskTag::EGameThread);
+
+	// Allows this program to accept a project argument on the commandline and use project-specific config
+	UE::ProjectUtilities::ParseProjectDirFromCommandline(ArgC, ArgV);
+
+	// start up the main loop,
+	GEngineLoop.PreInit(ArgC, ArgV);
+
+	int32 Result = 1;
+
+	FString FilePath;
+	FString EventName;
+
+	if (FParse::Value(FCommandLine::Get(), TEXT("csv="), FilePath)==false)
+	{
+		UE_LOG(LogCSVToTelemetry, Error, TEXT("Must provide a CSV file with -csv=[fillename]"));
+	}
+
+	if (FParse::Value(FCommandLine::Get(), TEXT("event="), EventName) == false)
+	{
+		UE_LOG(LogCSVToTelemetry, Error, TEXT("Must provide a Event name with -event=[eventname]"));
+	}
+	
+	if (FilePath.Len()>0 && EventName.Len()>0)
+	{
+		TArray<FString> Rows;
+
+		// Load all the rows from the file into memory
+		if (FFileHelper::LoadFileToStringArray(Rows, *FilePath))
+		{
+			UE_LOG(LogCSVToTelemetry, Display, TEXT("Importing %d rows"), Rows.Num());
+
+			FTelemetry Telemetry;
+
+			// Start the telemetry session and record each row as a single telemetry event 
+			if (Telemetry.StartSession())
+			{
+				uint32 TotalUploadSize=0;
+				uint32 RowIndex = 0;
+				uint32 KeySize = 0;
+				uint32 ValueSize = 0;
+	
+				TArray<FString> KeyArray;
+
+				FString Columns;
+
+				if (FParse::Value(FCommandLine::Get(), TEXT("columns="), Columns) == true)
+				{
+					// We have specified the CSV columns on the command line 
+					Columns.ParseIntoArray(KeyArray, TEXT("|"));
+					KeySize = Columns.Len();
+				}
+
+				for (const FString& Row : Rows)
+				{
+					if (KeyArray.IsEmpty())
+					{
+						// Parse the keys on the first row
+						Row.ParseIntoArray(KeyArray, TEXT(","));
+						KeySize = Row.Len();
+
+						for (int32 i = 0; i < KeyArray.Num(); ++i)
+						{
+							KeyArray[i].RemoveSpacesInline();
+						}
+					}
+					else
+					{
+						// Parse the values from subsequent rows
+						TArray<FString> ValueArray;
+						Row.ParseIntoArray(ValueArray, TEXT(","));
+
+						// Keys and values count should always be the same
+						if (KeyArray.Num() == ValueArray.Num())
+						{
+							ValueSize = Row.Len();
+
+							TArray<FAnalyticsEventAttribute> Attributes;
+
+							// Add the Key/Value pair to the attribute list
+							for (int32 i = 0; i < KeyArray.Num(); ++i)
+							{
+								ValueArray[i].RemoveSpacesInline();
+								Attributes.Emplace(KeyArray[i], ValueArray[i]);
+							}
+
+							// Keep track of the total data we're uploading
+							TotalUploadSize += KeySize + ValueSize;
+
+							// Now we have a complete event so send it
+							Telemetry.RecordEvent(EventName, Attributes);
+						}
+						else
+						{
+							// The key/value counts don't match
+							UE_LOG(LogCSVToTelemetry, Warning, TEXT("Row %d contains the incorrect value count of %d ( expected %d ) and will be skipped."), RowIndex, ValueArray.Num(), KeyArray.Num());
+						}
+					}
+
+					RowIndex++;
+				}
+
+				UE_LOG(LogCSVToTelemetry, Display, TEXT("Uploading %d bytes of event data"), TotalUploadSize);
+
+				Result = 0;
+			}
+
+			Telemetry.EndSession();
+		}
+		else
+		{
+			UE_LOG(LogCSVToTelemetry, Error, TEXT("Unable to read rows from CSV file %s"), *FilePath);
+		}
+	}
+
+	if (Result!=0)
+	{
+		// Always show the usage example if we have not been successful
+		UE_LOG(LogCSVToTelemetry, Display, TEXT("Usage CSVToTelemetry.exe -csv=[filename] -server=[severURL] -event=[eventname]"));
+	}
+
+	if (FParse::Param(FCommandLine::Get(), TEXT("fastexit")))
+	{
+		FPlatformMisc::RequestExitWithStatus(true, Result);
+	}
+
+	GLog->Flush();
+
+	RequestEngineExit(TEXT("CSVToTelemetry Exiting"));
+
+	FEngineLoop::AppPreExit();
+	FModuleManager::Get().UnloadModulesAtShutdown();
+	FEngineLoop::AppExit();
+
+	return Result;
+}
\ No newline at end of file
diff --git a/Engine/Source/Programs/CSVTools/CSVToTelemetry/Private/CSVToTelemetry.h b/Engine/Source/Programs/CSVTools/CSVToTelemetry/Private/CSVToTelemetry.h
new file mode 100644
--- /dev/null
+++ b/Engine/Source/Programs/CSVTools/CSVToTelemetry/Private/CSVToTelemetry.h
@@ -0,0 +1,5 @@
+// Copyright Epic Games, Inc. All Rights Reserved.
+
+#pragma once
+
+#include "CoreMinimal.h"
\ No newline at end of file
