CSV derived stats: add a framedelta function for computing the frame to frame delta between stat values. This can be used to convert cumulative stats into per-frame

diff --git a/Engine/Source/Programs/CSVTools/CsvStats/CsvStats.cs b/Engine/Source/Programs/CSVTools/CsvStats/CsvStats.cs
--- a/Engine/Source/Programs/CSVTools/CsvStats/CsvStats.cs
+++ b/Engine/Source/Programs/CSVTools/CsvStats/CsvStats.cs
@@ -512,41 +512,6 @@ namespace CSVStats
 			return metaData.GetValue(key, null);
 		}
 
-		private string ExpandCustomFunctions(string expression)
-		{
-			char[] operators = { '+', '-', '*', '/' };
-
-			string[] customFunctionNames = { "add" };
-			foreach (string funcName in customFunctionNames)
-			{
-				int startIdx = expression.IndexOf(funcName+"(");
-				if (startIdx == -1)
-				{
-					continue;
-				}
-				if (startIdx > 0 && !operators.Contains(expression[startIdx - 1]))
-				{
-					continue;
-				}
-				// FIXME: This will fail with nested brackets
-				int endIdx = expression.IndexOf(')', startIdx+1);
-				if (endIdx == -1)
-				{
-					continue;
-				}
-				string inner = expression[(startIdx+funcName.Length+1)..(endIdx)];
-				string[] innerParams = inner.Split(',');
-				if (funcName=="add")
-				{
-					// add(a,b,c) expands to (a+b+c). This is useful with wildcard expansion, since wildcards expand to comma separated list
-					string replaceStr = "(" + String.Join("+", innerParams) + ")";
-					expression = expression[..startIdx]+replaceStr+expression[(endIdx+1)..];
-				}
-				// Other functions can go here
-			}
-			return expression;
-		}
-
 		// Matches 
 		string ReadQuotedString(string str)
 		{
@@ -597,7 +562,7 @@ namespace CSVStats
 			return null;
 		}
 
-		enum BuiltinFunctionType { addstats, meta, eventcount, COUNT };
+		enum BuiltinFunctionType { addstats, meta, eventcount, framedelta, COUNT };
 		string ReadNextToken(string expression, ref int index, Dictionary<string, StatSamples> statsOut, int tokenIndex)
 		{
 			// Skip leading whitespace
@@ -658,7 +623,7 @@ namespace CSVStats
 						eventStr = (eventStr == null) ? paramString : eventStr;
 						if (eventStr == null)
 						{
-							throw new Exception("Warning: eventstr requires a quote '' arg: " + expression[oldIndex..index]);
+							throw new Exception("Warning: eventstr requires an arg: " + expression[oldIndex..index]);
 						}
 						string statName = "__STAT__EVENTCOUNT__" + tokenIndex.ToString();
 						statsOut[statName] = MakeEventCountStat(statName, eventStr);
@@ -686,6 +651,30 @@ namespace CSVStats
 						// Return the stats added together
 						return "("+String.Join("+", safeStatNames)+")";
 					}
+					else if (funcType == BuiltinFunctionType.framedelta)
+					{
+						// Quotes are optional. They're supported for consistency with stats
+						string srcStatName = ReadQuotedString(paramString);
+						srcStatName = (srcStatName == null) ? paramString : srcStatName;
+						if (srcStatName == null)
+						{
+							throw new Exception("Warning: frameDelta requires an arg: " + expression[oldIndex..index]);
+						}
+
+						StatSamples srcStat = GetStat(srcStatName);
+						if (srcStat == null)
+						{
+							throw new Exception("Warning: source stat not found: " + expression[oldIndex..index]);
+						}
+						string destStatName = "__STAT__FRAMEDELTA__" + tokenIndex.ToString();
+						StatSamples newStat = new StatSamples(destStatName, srcStat.samples.Count);
+						for (int j=1; j<srcStat.samples.Count; j++)
+						{
+							newStat.samples[j] = srcStat.samples[j] - srcStat.samples[j-1];
+						}
+						statsOut[destStatName] = newStat;
+						return destStatName;
+					}
 				}
 			}
 
