Added Session Label commandline option to the CSV Telemetry upload

diff --git a/Engine/Source/Programs/CSVTools/CSVToTelemetry/Private/CSVToTelemetry.cpp b/Engine/Source/Programs/CSVTools/CSVToTelemetry/Private/CSVToTelemetry.cpp
--- a/Engine/Source/Programs/CSVTools/CSVToTelemetry/Private/CSVToTelemetry.cpp
+++ b/Engine/Source/Programs/CSVTools/CSVToTelemetry/Private/CSVToTelemetry.cpp
@@ -63,8 +63,13 @@ public:
 				TEXT("Build_Changelist"), BuildSettings::GetCurrentChangelist(),
 				TEXT("Build_BranchName"), FApp::GetBranchName().ToLower());
 
-			uint32 SchemaVersion(0);
+			FString SessionLabel;
+			if (FParse::Value(FCommandLine::Get(), TEXT("sessionlabel="), SessionLabel))
+			{
+				DefaultAttributes.Emplace(TEXT("SessionLabel"), SessionLabel);
+			}
 
+			uint32 SchemaVersion(0);
 			if (FParse::Value(FCommandLine::Get(), TEXT("schema="), SchemaVersion))
 			{
 				// Optional schema version for this event
@@ -196,14 +201,14 @@ INT32_MAIN_INT32_ARGC_TCHAR_ARGV()
 						Row.ParseIntoArray(ValueArray, TEXT(","));
 
 						// Keys and values count should always be the same
-						if (KeyArray.Num() == ValueArray.Num())
+						if (KeyArray.Num() >= ValueArray.Num())
 						{
 							ValueSize = Row.Len();
 
 							TArray<FAnalyticsEventAttribute> Attributes;
 
 							// Add the Key/Value pair to the attribute list
-							for (int32 i = 0; i < KeyArray.Num(); ++i)
+							for (int32 i = 0; i < ValueArray.Num(); ++i)
 							{
 								ValueArray[i].RemoveSpacesInline();
 								Attributes.Emplace(KeyArray[i], ValueArray[i]);
