Changed CSVToTelemetry Tool to use StudioAnalytics engine library rather than reimplement  similar functionality AnalyticsET directly.
Means we can now data drive the analytics backends that the data is pushed to including the AnalyticsProviderHorde or a custom end user backend,

diff --git a/Engine/Source/Programs/CSVTools/CSVToTelemetry/CSVToTelemetry.Build.cs b/Engine/Source/Programs/CSVTools/CSVToTelemetry/CSVToTelemetry.Build.cs
--- a/Engine/Source/Programs/CSVTools/CSVToTelemetry/CSVToTelemetry.Build.cs
+++ b/Engine/Source/Programs/CSVTools/CSVToTelemetry/CSVToTelemetry.Build.cs
@@ -19,6 +19,7 @@ public class CSVToTelemetry : ModuleRules
 			"HTTP",
 			"Projects",
 			"RSA",
+			"StudioTelemetry",
 			"TargetPlatform"
 		});
 	}
diff --git a/Engine/Source/Programs/CSVTools/CSVToTelemetry/Private/CSVToTelemetry.cpp b/Engine/Source/Programs/CSVTools/CSVToTelemetry/Private/CSVToTelemetry.cpp
--- a/Engine/Source/Programs/CSVTools/CSVToTelemetry/Private/CSVToTelemetry.cpp
+++ b/Engine/Source/Programs/CSVTools/CSVToTelemetry/Private/CSVToTelemetry.cpp
@@ -4,127 +4,15 @@
 #include "RequiredProgramMainCPPInclude.h"
 #include "Analytics.h"
 #include "AnalyticsET.h"
-#include "BuildSettings.h"
-#include "Horde.h"
-#include "HttpModule.h"
-#include "HttpManager.h"
-#include "IAnalyticsProviderET.h"
+#include "StudioTelemetry.h"
 #include "Logging/LogMacros.h"
 #include "ProjectUtilities.h"
 
+
 DEFINE_LOG_CATEGORY_STATIC(LogCSVToTelemetry, Log, All);
 IMPLEMENT_APPLICATION(CSVToTelemetry, "CSVToTelemetry");
 
-class FTelemetry
-{
-public:
-
-	bool StartSession()
-	{
-		FString APIServer;
-
-		if (AnalyticsProvider.IsValid())
-		{
-			// Already running the telemetry
-			return false;
-		}
-		
-		// We must have supplied at least the APIServer URL to send any telemetry
-		if (FParse::Value(FCommandLine::Get(), TEXT("server="), APIServer)==false)
-		{
-			UE_LOG(LogCSVToTelemetry, Error, TEXT("Must provide a valid server address with -server=[serverURL]"));
-			return false;
-		}
-
-		// Allow override of API Endpoint
-		FString APIEndpoint = TEXT("datarouter/api/v1/public/data");
-		FParse::Value(FCommandLine::Get(), TEXT("endpoint="), APIEndpoint);
-
-		// Allow override of API Endpoint
-		FString APIKey = TEXT("StudioAnalytics.Dev");
-		FParse::Value(FCommandLine::Get(), TEXT("apikey="), APIKey);
-
-		FAnalyticsET::Config Config;
-
-		Config.APIServerET = MoveTemp(APIServer);
-		Config.APIEndpointET = MoveTemp(APIEndpoint);
-		Config.APIKeyET = MoveTemp(APIKey);
-		Config.AppEnvironment = TEXT("datacollector-binary");
-		Config.AppVersionET = FEngineVersion::Current().ToString();
-
-		AnalyticsProvider = FAnalyticsET::Get().CreateAnalyticsProvider(Config);
-
-		if (AnalyticsProvider.IsValid())
-		{
-			// Prepare the default attributes appended to all events
-			TArray<FAnalyticsEventAttribute> DefaultAttributes = MakeAnalyticsEventAttributeArray(					
-				TEXT("Horde_TemplateName"), FHorde::GetTemplateName(),
-				TEXT("Horde_StepName"), FHorde::GetStepName(),
-				TEXT("Build_Changelist"), BuildSettings::GetCurrentChangelist(),
-				TEXT("Build_BranchName"), FApp::GetBranchName().ToLower());
-
-			FString SessionLabel;
-			if (FParse::Value(FCommandLine::Get(), TEXT("sessionlabel="), SessionLabel))
-			{
-				DefaultAttributes.Emplace(TEXT("SessionLabel"), SessionLabel);
-			}
-
-			uint32 SchemaVersion(0);
-			if (FParse::Value(FCommandLine::Get(), TEXT("schema="), SchemaVersion))
-			{
-				// Optional schema version for this event
-				DefaultAttributes.Emplace(TEXT("SchemaVersion"), SchemaVersion);
-			}
-
-			// Set the default attributes
-			AnalyticsProvider->SetDefaultEventAttributes(MoveTemp(DefaultAttributes));
-
-			// Start the provider session
-			AnalyticsProvider->StartSession();
-
-			return true;
-		}
-		
-		return false;
-	}
-
-	void RecordEvent(const FString& EventName, const TArray<FAnalyticsEventAttribute>& Attributes)
-	{
-		if (AnalyticsProvider.IsValid())
-		{
-			AnalyticsProvider->RecordEvent(EventName, Attributes);
-		}
-	}
-
-	void EndSession()
-	{
-		if (AnalyticsProvider.IsValid())
-		{		
-			// End the session
-			AnalyticsProvider->FlushEvents();
-			AnalyticsProvider->EndSession();
-			AnalyticsProvider.Reset();
-
-			UE_LOG(LogCSVToTelemetry, Display, TEXT("Waiting for Telemetry upload to complete"));
-
-			double StartTime = StartTime = FPlatformTime::Seconds();
-			// It is quite likely that one of the analytics providers is sending data via the FHttpManager so we
-			// need to flush that as well to make sure that the message is sent. 
-			if (FHttpModule* HttpModule = FModuleManager::GetModulePtr<FHttpModule>("HTTP"))
-			{
-				FHttpManager& HttpManager = HttpModule->GetHttpManager();
-				HttpManager.Flush(EHttpFlushReason::FullFlush);
-			}
-
-			UE_LOG(LogCSVToTelemetry, Display, TEXT("Telemetry upload finished in %f seconds"), FPlatformTime::Seconds() - StartTime);
-		}
-	}
-
-private:
 
-	TSharedPtr<IAnalyticsProviderET>	AnalyticsProvider;
-	float								TimeOut = 1200.0f;
-};
 
 INT32_MAIN_INT32_ARGC_TCHAR_ARGV()
 {
@@ -158,22 +46,25 @@ INT32_MAIN_INT32_ARGC_TCHAR_ARGV()
 		// Load all the rows from the file into memory
 		if (FFileHelper::LoadFileToStringArray(Rows, *FilePath))
 		{
-			UE_LOG(LogCSVToTelemetry, Display, TEXT("Importing %d rows"), Rows.Num());
+			UE_LOG(LogCSVToTelemetry, Display, TEXT("Imported %d rows from file"), Rows.Num());
 
-			FTelemetry Telemetry;
+			FStudioTelemetry::Get().StartSession();
 
 			// Start the telemetry session and record each row as a single telemetry event 
-			if (Telemetry.StartSession())
+			if (FStudioTelemetry::Get().IsSessionRunning())
 			{
-				uint32 TotalUploadSize=0;
+				uint32 TotalUploadSize = 0;
 				uint32 RowIndex = 0;
 				uint32 KeySize = 0;
 				uint32 ValueSize = 0;
-	
+
 				TArray<FString> KeyArray;
 
 				FString Columns;
 
+				uint32 SchemaVersion(0);
+				FParse::Value(FCommandLine::Get(), TEXT("schema="), SchemaVersion);
+				
 				if (FParse::Value(FCommandLine::Get(), TEXT("columns="), Columns) == true)
 				{
 					// We have specified the CSV columns on the command line 
@@ -207,18 +98,24 @@ INT32_MAIN_INT32_ARGC_TCHAR_ARGV()
 
 							TArray<FAnalyticsEventAttribute> Attributes;
 
+							// If we've specified a non-zero schema version then add that
+							if (SchemaVersion != 0)
+							{
+								Attributes.Emplace(TEXT("SchemaVersion"), SchemaVersion);
+							}
+
 							// Add the Key/Value pair to the attribute list
 							for (int32 i = 0; i < ValueArray.Num(); ++i)
 							{
 								ValueArray[i].RemoveSpacesInline();
 								Attributes.Emplace(KeyArray[i], ValueArray[i]);
 							}
-
+	
 							// Keep track of the total data we're uploading
 							TotalUploadSize += KeySize + ValueSize;
 
 							// Now we have a complete event so send it
-							Telemetry.RecordEvent(EventName, Attributes);
+							FStudioTelemetry::Get().RecordEvent(EventName, Attributes);
 						}
 						else
 						{
@@ -230,12 +127,12 @@ INT32_MAIN_INT32_ARGC_TCHAR_ARGV()
 					RowIndex++;
 				}
 
-				UE_LOG(LogCSVToTelemetry, Display, TEXT("Uploading %d bytes of event data"), TotalUploadSize);
+				UE_LOG(LogCSVToTelemetry, Display, TEXT("Generated %d bytes of event data"), TotalUploadSize);
+
+				FStudioTelemetry::Get().EndSession();
 
 				Result = 0;
 			}
-
-			Telemetry.EndSession();
 		}
 		else
 		{
@@ -246,7 +143,7 @@ INT32_MAIN_INT32_ARGC_TCHAR_ARGV()
 	if (Result!=0)
 	{
 		// Always show the usage example if we have not been successful
-		UE_LOG(LogCSVToTelemetry, Display, TEXT("Usage CSVToTelemetry.exe -csv=[filename] -server=[severURL] -event=[eventname]"));
+		UE_LOG(LogCSVToTelemetry, Display, TEXT("Usage CSVToTelemetry.exe -csv=[filename] -event=[eventname] ( -schema=[value] -columns=[name1|name2|....] )"));
 	}
 
 	if (FParse::Param(FCommandLine::Get(), TEXT("fastexit")))
