Fixes for minimap issues in PRT reports
Add support for future minimap transform metdata
Add backcompat for old metadata

diff --git a/Engine/Source/Programs/CSVTools/PerfReportTool/Summaries/MapOverlaySummary.cs b/Engine/Source/Programs/CSVTools/PerfReportTool/Summaries/MapOverlaySummary.cs
--- a/Engine/Source/Programs/CSVTools/PerfReportTool/Summaries/MapOverlaySummary.cs
+++ b/Engine/Source/Programs/CSVTools/PerfReportTool/Summaries/MapOverlaySummary.cs
@@ -82,9 +82,26 @@ namespace PerfSummaries
 				sourceImagePath = System.IO.Path.GetFullPath(System.IO.Path.Combine(baseXmlDirectory, sourceImagePath));
 			}
 
-			offsetX = element.GetSafeAttribute<float>(vars, "offsetX", 0.0f);
-			offsetY = element.GetSafeAttribute<float>(vars, "offsetY", 0.0f);
-			scale = element.GetSafeAttribute<float>(vars, "scale", 1.0f);
+			string transformtext = element.GetSafeAttribute(vars, "transform", "");
+			if (transformtext.Length == 0)
+			{
+				// support for previous metadata
+				float offsetX = element.GetSafeAttribute<float>(vars, "offsetX", 0.0f);
+				float offsetY = element.GetSafeAttribute<float>(vars, "offsetY", 0.0f);
+				float scale = element.GetSafeAttribute<float>(vars, "scale", 1.0f);
+
+				transform[0] = scale * 0.5f;
+				transform[1] = 0.0f;
+				transform[2] = 0.0f;
+				transform[3] = -scale * 0.5f;
+				transform[4] = offsetX * 0.5f + 0.5f;
+				transform[5] = offsetY * 0.5f + 0.5f;
+			}
+			else
+			{
+				transform = Array.ConvertAll(transformtext.Split(',', StringSplitOptions.RemoveEmptyEntries), float.Parse);
+			}
+
 			title = element.GetSafeAttribute(vars, "title", "Events");
 			destImageFilename = element.GetRequiredAttribute<string>(vars, "destImage");
 			imageWidth = element.GetSafeAttribute<float>(vars, "width", 250.0f);
@@ -107,14 +124,14 @@ namespace PerfSummaries
 
 		int toSvgX(float worldX, float worldY)
 		{
-			float svgX = (worldY * scale + offsetX) * 0.5f + 0.5f;
+			float svgX = worldX * transform[0] + worldY * transform[1] + transform[4];
 			svgX *= imageWidth;
 			return (int)(svgX + 0.5f);
 		}
 
 		int toSvgY(float worldX, float worldY)
 		{
-			float svgY = 1.0f - (worldX * scale + offsetY) * 0.5f - 0.5f;
+			float svgY = worldX * transform[2] + worldY * transform[3] + transform[5];
 			svgY *= imageHeight;
 			return (int)(svgY + 0.5f);
 		}
@@ -302,9 +319,7 @@ namespace PerfSummaries
 		}
 		string title;
 		string sourceImagePath;
-		float offsetX;
-		float offsetY;
-		float scale;
+		float[] transform = new float[6];
 		string destImageFilename;
 		float imageWidth;
 		float imageHeight;
