Minor update to change delimitor from comma to space to avoid limitation in CSV

diff --git a/Engine/Source/Programs/CSVTools/PerfReportTool/Summaries/MapOverlaySummary.cs b/Engine/Source/Programs/CSVTools/PerfReportTool/Summaries/MapOverlaySummary.cs
--- a/Engine/Source/Programs/CSVTools/PerfReportTool/Summaries/MapOverlaySummary.cs
+++ b/Engine/Source/Programs/CSVTools/PerfReportTool/Summaries/MapOverlaySummary.cs
@@ -99,7 +99,7 @@ namespace PerfSummaries
 			}
 			else
 			{
-				transform = Array.ConvertAll(transformtext.Split(',', StringSplitOptions.RemoveEmptyEntries), float.Parse);
+				transform = Array.ConvertAll(transformtext.Split(' ', StringSplitOptions.RemoveEmptyEntries), float.Parse);
 			}
 
 			title = element.GetSafeAttribute(vars, "title", "Events");
