Updated perf tooling
Fixed sign error on Y translation

diff --git a/Engine/Source/Programs/CSVTools/PerfReportTool/Summaries/MapOverlaySummary.cs b/Engine/Source/Programs/CSVTools/PerfReportTool/Summaries/MapOverlaySummary.cs
--- a/Engine/Source/Programs/CSVTools/PerfReportTool/Summaries/MapOverlaySummary.cs
+++ b/Engine/Source/Programs/CSVTools/PerfReportTool/Summaries/MapOverlaySummary.cs
@@ -95,7 +95,7 @@ namespace PerfSummaries
 				transform[2] = 0.0f;
 				transform[3] = -scale * 0.5f;
 				transform[4] = offsetX * 0.5f + 0.5f;
-				transform[5] = offsetY * 0.5f + 0.5f;
+				transform[5] = -offsetY * 0.5f + 0.5f;
 			}
 			else
 			{
