PerfReportTool 4.243.0 - Add support for external metadata:
-Use the -externalMetadataSources param to specify paths for external metadata. This is applied to summary tables dynamically. It is not cached
- Each supplied path can include multiple metadata keys specified with curly braces, e.g "{buildversion}/{platform}.txt". The final filenames are resolved at runtime based on metadata for each row
- If the files exist, they will be read as key=value text files (similar to ini files) and added to the row data as ExternalMetadata entries
- The SummaryTableElement data has the type ExternalMetadata (prefix "extmeta")

[RN] PerfReportTool 4.243.0: Add support for external metadata using the -externalMetadataSources param. This is applied to summary tables dynamically. Each supplied path can include multiple metadata keys specified with curly braces, e.g "{buildversion}/{platform}.txt". The filenames are resolved at runtime based on metadata for each row. These files should contain key=value pairs containing the metadata entries

diff --git a/Engine/Source/Programs/CSVTools/PerfReportTool/CsvFileCache..cs b/Engine/Source/Programs/CSVTools/PerfReportTool/CsvFileCache..cs
--- a/Engine/Source/Programs/CSVTools/PerfReportTool/CsvFileCache..cs
+++ b/Engine/Source/Programs/CSVTools/PerfReportTool/CsvFileCache..cs
@@ -75,7 +75,7 @@ namespace PerfReportTool
 					}
 					else
 					{
-						Console.WriteLine("CSV file " + inFilename + " is contains no lines!");
+						Console.WriteLine("CSV file " + inFilename + " contains no lines!");
 						dummyCsvStats = new CsvStats();
 					}
 				}
@@ -86,21 +86,7 @@ namespace PerfReportTool
 					cachedSummaryTableRowData = SummaryTableRowData.TryReadFromCacheFile(inFilename, true);
 					if (cachedSummaryTableRowData != null)
 					{
-						foreach (KeyValuePair<string, SummaryTableElement> entry in cachedSummaryTableRowData.dict)
-						{
-							if (entry.Value.type == SummaryTableElement.Type.CsvMetadata)
-							{
-								if (metadata == null)
-								{
-									metadata = new CsvMetadata();
-								}
-								metadata.Values.Add(entry.Key, entry.Value.value);
-							}
-							if (entry.Value.type == SummaryTableElement.Type.CsvStatAverage)
-							{
-								break;
-							}
-						}
+						metadata = cachedSummaryTableRowData.ReadCsvMetadata();
 					}
 					else
 					{
@@ -243,9 +229,118 @@ namespace PerfReportTool
 		public ReportTypeInfo reportTypeInfo;
 		public XmlVariableMappings xmlVariableMappings;
 		public DerivedCsvStatDefinition[] derivedCsvStatDefinitions;
+
 		FileType fileType;
 	};
 
+	class ExternalMetadata : Dictionary<string, string>
+	{
+
+	}
+
+	class ExternalMetadataCache
+	{
+		public ExternalMetadataCache(List<string> inExternalMetadataSourceStrings)
+		{
+			externalMetadataSourceStrings = inExternalMetadataSourceStrings;
+		}
+		public void InjectExternalMetadata(SummaryTableRowData summaryTableRowData)
+		{
+			foreach (string filename in externalMetadataSourceStrings)
+			{
+				string resolvedFilename = ResolveExternalMetadataFilename(filename, summaryTableRowData);
+				if (resolvedFilename == null)
+				{
+					continue;
+				}
+				ExternalMetadata fileDict = LoadExternalMetadataFile(resolvedFilename);
+				foreach (KeyValuePair<string, string> entry in fileDict)
+				{
+					// Add metadata if it doesn't already exist
+					SummaryTableElement element = summaryTableRowData.Get(entry.Key);
+					if (element == null)
+					{
+						summaryTableRowData.Add(SummaryTableElement.Type.ExternalMetadata, entry.Key, entry.Value);
+					}
+				}
+			}
+		}
+
+		ExternalMetadata LoadExternalMetadataFile(string filename)
+		{
+			lock (cachedExternalMetadataFiles)
+			{
+				if (cachedExternalMetadataFiles.TryGetValue(filename, out ExternalMetadata cachedDict))
+				{
+					return cachedDict;
+				}
+			}
+			ExternalMetadata dict = new ExternalMetadata();
+			if (System.IO.File.Exists(filename))
+			{
+				string[] lines = System.IO.File.ReadAllLines(filename);
+				foreach (string line in lines)
+				{
+					int idx = line.IndexOf('=');
+					if (idx == -1)
+					{
+						break;
+					}
+					string key = line[..idx].ToLowerInvariant();
+					string value = line[(idx + 1)..];
+					dict[key] = value;
+				}
+			}
+			// Add to the cache if it wasn't added while we were reading the file
+			lock (cachedExternalMetadataFiles)
+			{
+				if (!cachedExternalMetadataFiles.ContainsKey(filename))
+				{
+					cachedExternalMetadataFiles[filename] = dict;
+				}				
+			}			
+			return dict;
+		}
+
+
+		// Converts an external metadata file path (including {} metadata entries) to an actual filename by resolving the metadata entries
+		// Returns null on failure
+		static string ResolveExternalMetadataFilename(string sourceFilename, SummaryTableRowData summaryTableRowData)
+		{
+			string filenameOut = "";
+			// Replace {meta} entries with actual metadata
+			int readIndex = 0;
+			int startIndex = sourceFilename.IndexOf('{');
+			while (startIndex >= 0)
+			{
+				int endIndex = sourceFilename.IndexOf('}', startIndex);
+				if (endIndex == -1)
+				{
+					throw new Exception("Warning: bad external metadata filename (mismatched {): " + sourceFilename);
+				}
+				string key = sourceFilename[(startIndex+1)..endIndex].ToLowerInvariant();
+				
+				SummaryTableElement sourceElement = summaryTableRowData.Get(key);
+				if (sourceElement == null || sourceElement.type != SummaryTableElement.Type.CsvMetadata)
+				{
+					// If we fail to find the metadata, that's probably fine. Just return null. Ignore other types
+					return null;
+				}
+				string value = sourceElement.value;
+				filenameOut += sourceFilename[readIndex..startIndex] + value;
+				readIndex = endIndex+1;
+
+				startIndex = sourceFilename.IndexOf('{', endIndex);
+			}
+			filenameOut += sourceFilename[readIndex..];
+			return filenameOut;
+		}
+
+
+		Dictionary<string, ExternalMetadata> cachedExternalMetadataFiles = new Dictionary<string, ExternalMetadata>();
+		List<string> externalMetadataSourceStrings;
+	};
+
 	class CsvFileCache
 	{
 		public CsvFileCache(
@@ -262,7 +357,8 @@ namespace PerfReportTool
 			bool inRemoveDuplicates,
 			bool inRequireMetadata,
 			string inSummaryTableCacheDir,
-			bool inListFilesMode)
+			bool inListFilesMode,
+			ExternalMetadataCache inExternalMetadataCache)
 		{
 			csvFileInfos = new CsvFileInfo[inCsvFilenames.Length];
 			for (int i = 0; i < inCsvFilenames.Length; i++)
@@ -291,6 +387,8 @@ namespace PerfReportTool
 			csvIdToFilename = new Dictionary<string, string>();
 			outCsvQueue = new BlockingCollection<CsvFileInfo>(lookaheadCount);
 
+			externalMetadataCache = inExternalMetadataCache;
+
 			// Kick off the workers (must be done last)
 			if (inLookaheadCount > 0)
 			{
@@ -435,6 +533,12 @@ namespace PerfReportTool
 							fileInfo.isValid = true;
 						}
 					}
+					// Inject external metadata into the cached row data
+					if (file.cachedSummaryTableRowData != null)
+					{
+						externalMetadataCache.InjectExternalMetadata(file.cachedSummaryTableRowData);
+					}
+
 				}
 				fileInfo.cachedFile = file;
 				// TODO: only add valid files to the queue? Need alternative stopping condition
@@ -504,5 +608,6 @@ namespace PerfReportTool
 		}
 		CsvFileInfo[] csvFileInfos;
 
+		ExternalMetadataCache externalMetadataCache;
 	}
 }
\ No newline at end of file
diff --git a/Engine/Source/Programs/CSVTools/PerfReportTool/PerfReportTool.cs b/Engine/Source/Programs/CSVTools/PerfReportTool/PerfReportTool.cs
--- a/Engine/Source/Programs/CSVTools/PerfReportTool/PerfReportTool.cs
+++ b/Engine/Source/Programs/CSVTools/PerfReportTool/PerfReportTool.cs
@@ -23,7 +23,7 @@ namespace PerfReportTool
     class Version
     {
 		// Format: Major.Minor.Bugfix
-        private static string VersionString = "4.242.0";
+        private static string VersionString = "4.243.0";
 
         public static string Get() { return VersionString; }
     };
@@ -175,6 +175,8 @@ namespace PerfReportTool
 			"  -summaryTableXmlRowSortAppend <list,of,stats> : append these stats to the summary table's row sort list\n" +
 			"  -transposeTable : write the summary tables transposed\n" +
 			"  -transposeCollatedTable : write the collated summary table transposed (disables min/max columns)\n" +
+			"  -externalMetadataSources <list of paths>: specify paths for external metadata key=value text files\n" +
+			"	   these resolve from existing metadata to form the final path. e.g D:\\ExtMeta\\{{platform}}\\{{buildversion}}.txt\n" +
 			"Diff rows\n" +
 			"  -addDiffRows : adds diff rows after the first two rows\n" +
 			"  -sortColumnsByDiff : sorts columns by the max of its diff scores (use with -addDiffRows).\n" +
@@ -586,6 +588,8 @@ namespace PerfReportTool
 				writeDetailedReports = false;
 			}
 
+			List<string> externalMetadataSources = GetListArg("externalMetadataSources");
+			ExternalMetadataCache externalMetadataCache = new ExternalMetadataCache(externalMetadataSources);
 
 			CsvFileCache csvFileCache = new CsvFileCache(
 				csvFilenames,
@@ -601,7 +605,8 @@ namespace PerfReportTool
 				bRemoveDuplicates,
 				bRequireMetadata,
 				summaryTableCacheForRead,
-				bListFilesMode);
+				bListFilesMode,
+				externalMetadataCache);
 
 			SummaryTable summaryTable = new SummaryTable();
 			bool bWriteToSummaryTableCache = summaryTableCacheDir != null && !bSummaryTableCacheReadonly;
@@ -673,13 +678,18 @@ namespace PerfReportTool
 
 								}
 
-								if (rowData != null && bWriteToSummaryTableCache)
+								if (rowData != null)
 								{
-									if (rowData.WriteToCache(summaryTableCacheDir, cachedCsvFile.summaryTableCacheId))
+									// Inject external metadata if we have it
+									externalMetadataCache.InjectExternalMetadata(rowData);
+									if (bWriteToSummaryTableCache)
 									{
-										Console.WriteLine("Cached summary rowData for CSV: " + cachedCsvFile.filename);
-										summaryTableCacheStats.WriteCount++;
-										perfLog.LogTiming("  WriteSummaryTableCache");
+										if (rowData.WriteToCache(summaryTableCacheDir, cachedCsvFile.summaryTableCacheId))
+										{
+											Console.WriteLine("Cached summary rowData for CSV: " + cachedCsvFile.filename);
+											summaryTableCacheStats.WriteCount++;
+											perfLog.LogTiming("  WriteSummaryTableCache");
+										}
 									}
 								}
 							}
@@ -1354,7 +1364,6 @@ namespace PerfReportTool
 						rowData.Add(SummaryTableElement.Type.CsvStatAverage, stat.Name, (double)stat.average);
 					}
 				}
-
 			}
 
 			if (htmlFilename != null && !string.IsNullOrEmpty(outputDir))
diff --git a/Engine/Source/Programs/CSVTools/PerfReportTool/SummaryTable.cs b/Engine/Source/Programs/CSVTools/PerfReportTool/SummaryTable.cs
--- a/Engine/Source/Programs/CSVTools/PerfReportTool/SummaryTable.cs
+++ b/Engine/Source/Programs/CSVTools/PerfReportTool/SummaryTable.cs
@@ -1206,8 +1206,10 @@ namespace PerfSummaries
 						return 3;
 					case SummaryTableElement.Type.ToolMetadata:
 						return 4;
-					default:
+					case SummaryTableElement.Type.ExternalMetadata:
 						return 5;
+					default:
+						return 6;
 				}
 			}
 
@@ -1452,6 +1454,8 @@ namespace PerfSummaries
 					return "[meta]";
 				case SummaryTableElement.Type.ToolMetadata:
 					return "[toolmeta]";
+				case SummaryTableElement.Type.ExternalMetadata:
+					return "[extmeta]";
 				default:
 					return "[invalid]";
 			}
@@ -1478,6 +1482,10 @@ namespace PerfSummaries
 				{
 					return SummaryTableElement.Type.ToolMetadata;
 				}
+				else if (fullNameWithQualifier.StartsWith("[extmeta]"))
+				{
+					return SummaryTableElement.Type.ExternalMetadata;
+				}
 			}
 			return SummaryTableElement.Type.COUNT;
 		}
@@ -1510,6 +1518,10 @@ namespace PerfSummaries
 				{
 					return SummaryTableElement.Type.ToolMetadata;
 				}
+				else if (fullNameWithQualifier.StartsWith("[extmeta]"))
+				{
+					return SummaryTableElement.Type.ExternalMetadata;
+				}
 			}
 			statNameOut = fullNameWithQualifier;
 			return SummaryTableElement.Type.COUNT;
diff --git a/Engine/Source/Programs/CSVTools/PerfReportTool/SummaryTableData.cs b/Engine/Source/Programs/CSVTools/PerfReportTool/SummaryTableData.cs
--- a/Engine/Source/Programs/CSVTools/PerfReportTool/SummaryTableData.cs
+++ b/Engine/Source/Programs/CSVTools/PerfReportTool/SummaryTableData.cs
@@ -8,6 +8,7 @@ using PerfReportTool;
 using System.IO;
 using System.Text.Json;
 using System.Threading.Tasks;
+using CSVStats;
 
 namespace PerfSummaries
 {
@@ -106,6 +107,7 @@ namespace PerfSummaries
 			CsvMetadata,
 			SummaryTableMetric,
 			ToolMetadata,
+			ExternalMetadata,
 			COUNT,
 			ANY=COUNT				
 		};
@@ -350,6 +352,28 @@ namespace PerfSummaries
 			return TryReadFromCacheFile(filename);
 		}
 
+		public CsvMetadata ReadCsvMetadata()
+		{
+			CsvMetadata metadata = null;
+			foreach (KeyValuePair<string, SummaryTableElement> entry in dict)
+			{
+				if (entry.Value.type == SummaryTableElement.Type.CsvMetadata)
+				{
+					if (metadata == null)
+					{
+						metadata = new CsvMetadata();
+					}
+					metadata.Values.Add(entry.Key, entry.Value.value);
+				}
+				if (entry.Value.type == SummaryTableElement.Type.CsvStatAverage)
+				{
+					// Metadata values appear first so we can early out
+					break;
+				}
+			}
+			return metadata;
+		}
+
 		public static SummaryTableRowData TryReadFromCacheFile(string filename, bool bReadJustInitialMetadata = false)
 		{
 			SummaryTableRowData metaData = null;
@@ -403,6 +427,21 @@ namespace PerfSummaries
 			return metaData;
 		}
 
+
+		// Filter out external entries for caching purposes
+		List<KeyValuePair<string, SummaryTableElement>> GetCacheableEntries()
+		{
+			List<KeyValuePair<string, SummaryTableElement>> cacheableEntries = new List<KeyValuePair<string, SummaryTableElement>>();
+			foreach (KeyValuePair<string, SummaryTableElement> entry in dict)
+			{
+				if (entry.Value.type != SummaryTableElement.Type.ExternalMetadata)
+				{
+					cacheableEntries.Add(entry);
+				}
+			}
+			return cacheableEntries;
+		}
+
 		public bool WriteToCache(string summaryTableCacheDir, string csvId)
 		{
 			string filename = Path.Combine(summaryTableCacheDir, csvId + ".prc");
@@ -414,8 +453,10 @@ namespace PerfSummaries
 					writer.Write(CacheVersion);
 					writer.Write(SummaryTableElement.CacheVersion);
 
-					writer.Write(dict.Count);
-					foreach (KeyValuePair<string, SummaryTableElement> entry in dict)
+					// Only write non-external entries to the cache
+					List<KeyValuePair<string, SummaryTableElement>> cacheableEntries = GetCacheableEntries();
+					writer.Write(cacheableEntries.Count);
+					foreach (KeyValuePair<string, SummaryTableElement> entry in cacheableEntries)
 					{
 						writer.Write(entry.Key);
 						entry.Value.WriteToCache(writer);
@@ -551,7 +592,8 @@ namespace PerfSummaries
 				var DataTypes = Enum.GetValues(typeof(SummaryTableElement.Type));
 				foreach (SummaryTableElement.Type dataType in DataTypes)
 				{
-					if (dataType != SummaryTableElement.Type.COUNT)
+					// External metadata is never cached
+					if (dataType != SummaryTableElement.Type.COUNT && dataType != SummaryTableElement.Type.ExternalMetadata)
 					{
 						DictOut[dataType.ToString()] = new Dictionary<string, dynamic>();
 					}
@@ -565,6 +607,11 @@ namespace PerfSummaries
 				{
 					continue;
 				}
+				// External metadata is never cached
+				if (Element.type == SummaryTableElement.Type.ExternalMetadata)
+				{
+					continue;
+				}
 				if (bWriteAllElementData)
 				{
 					DictOut[Element.type.ToString()][Element.name] = Element.ToJsonDict(false);
