PerfReportTool 4.243.1: Preserve case for external metadata

diff --git a/Engine/Source/Programs/CSVTools/PerfReportTool/CsvFileCache..cs b/Engine/Source/Programs/CSVTools/PerfReportTool/CsvFileCache..cs
--- a/Engine/Source/Programs/CSVTools/PerfReportTool/CsvFileCache..cs
+++ b/Engine/Source/Programs/CSVTools/PerfReportTool/CsvFileCache..cs
@@ -286,7 +286,7 @@ namespace PerfReportTool
 					{
 						break;
 					}
-					string key = line[..idx].ToLowerInvariant();
+					string key = line[..idx];
 					string value = line[(idx + 1)..];
 					dict[key] = value;
 				}
diff --git a/Engine/Source/Programs/CSVTools/PerfReportTool/PerfReportTool.cs b/Engine/Source/Programs/CSVTools/PerfReportTool/PerfReportTool.cs
--- a/Engine/Source/Programs/CSVTools/PerfReportTool/PerfReportTool.cs
+++ b/Engine/Source/Programs/CSVTools/PerfReportTool/PerfReportTool.cs
@@ -23,7 +23,7 @@ namespace PerfReportTool
     class Version
     {
 		// Format: Major.Minor.Bugfix
-        private static string VersionString = "4.243.0";
+        private static string VersionString = "4.243.1";
 
         public static string Get() { return VersionString; }
     };
