PerfReportTool 4.244.0 : Add support for collation of string and date values, so we can display this info in collated summary tables
- String and date values can now be optionally included in summary tables in the RHS stats section.
- Instead of displaying avg/min/max sub columns, collated string values will display if they match, otherwise they'll show "(multple)".
- Date values can be collated either by showing the newest value or the oldest value
- The behavior is controlled via XML or via the following cmdline parameters (where "hide" is equvalient to the old behavior)
-collatedStringVisibility <show|hide|auto> : sets visibility of strings in collated summary tables (default=auto)
-collatedDateVisibility <hide|newest|oldest> : sets visibility of dates in collated summary tables (default=newest)
- SummaryTable XML support is also provided, via "collatedDateVisibility" and "collatedStringVisibility" attributes (with the same options as above). The commandline params override the XML config
- Add enum parsing helper functions for both XML and commandline
- Refactor collation code

diff --git a/Engine/Source/Programs/CSVTools/CsvStats/CommandLineTool.cs b/Engine/Source/Programs/CSVTools/CsvStats/CommandLineTool.cs
--- a/Engine/Source/Programs/CSVTools/CsvStats/CommandLineTool.cs
+++ b/Engine/Source/Programs/CSVTools/CsvStats/CommandLineTool.cs
@@ -3,9 +3,7 @@
 
 using System;
 using System.Collections.Generic;
-using System.Linq;
 using System.Text;
-using System.Threading.Tasks;
 using System.IO;
 using System.Diagnostics;
 
@@ -282,6 +280,18 @@ namespace CSVStats
 			return commandLine.GetBoolArg(key, defaultValue);
         }
 
+		protected T GetEnumArg<T>(string key, T defaultValue) where T : System.Enum
+		{
+			string str = GetArg(key, null);
+#nullable enable
+			if ( str != null && Enum.TryParse(typeof(T), str, true, out object? result) )
+#nullable disable
+			{
+				return (T)result;
+			}
+			return defaultValue;
+		}
+
 		protected bool? GetOptionalBoolArg(string key)
 		{
 			return commandLine.GetOptionalBoolArg(key);
diff --git a/Engine/Source/Programs/CSVTools/PerfReportTool/PerfReportTool.cs b/Engine/Source/Programs/CSVTools/PerfReportTool/PerfReportTool.cs
--- a/Engine/Source/Programs/CSVTools/PerfReportTool/PerfReportTool.cs
+++ b/Engine/Source/Programs/CSVTools/PerfReportTool/PerfReportTool.cs
@@ -23,7 +23,7 @@ namespace PerfReportTool
     class Version
     {
 		// Format: Major.Minor.Bugfix
-        private static string VersionString = "4.243.1";
+        private static string VersionString = "4.244.0";
 
         public static string Get() { return VersionString; }
     };
@@ -153,7 +153,7 @@ namespace PerfReportTool
 			"  -reverseTable [0|1]: Reverses the order of summary tables (set 0 to force off)\n" +
 			"  -sortTrailingDigitsAsNumeric : detects trailing digits and pads them when sorting summary table\n"+
 			"  -scrollableTable [0|1]: makes the summary table scrollable, with frozen first rows and columns (set 0 to force off)\n" +
-			"  -colorizeTable [off|budget|auto]: selects the table colorization mode. If omitted, uses the default in the summary\n" +
+			"  -colorizeTable <off|budget|auto>: selects the table colorization mode. If omitted, uses the default in the summary\n" +
 			"     xml table if set.\n" +
 			"  -maxSummaryTableStringLength <n>: strings longer than this will get truncated\n" +
 			"  -allowDuplicateCSVs : doesn't remove duplicate CSVs (Note: can cause summary table cache file locking issues)\n" +
@@ -177,6 +177,9 @@ namespace PerfReportTool
 			"  -transposeCollatedTable : write the collated summary table transposed (disables min/max columns)\n" +
 			"  -externalMetadataSources <list of paths>: specify paths for external metadata key=value text files\n" +
 			"	   these resolve from existing metadata to form the final path. e.g D:\\ExtMeta\\{{platform}}\\{{buildversion}}.txt\n" +
+			"  -collatedStringVisibility <show|hide|auto> : sets visibility of strings in collated summary tables (default=auto)\n" +
+			"  -collatedDateVisibility <hide|newest|oldest> : sets visibility of dates in collated summary tables (default=newest)\n" +
+
 			"Diff rows\n" +
 			"  -addDiffRows : adds diff rows after the first two rows\n" +
 			"  -sortColumnsByDiff : sorts columns by the max of its diff scores (use with -addDiffRows).\n" +
@@ -882,26 +885,12 @@ namespace PerfReportTool
 			// If -colorizeTable isn't specified then we use the default from the xml.
 			// If -colorizeTable isn't specified and it's not set in the summary table xml it uses the default set in the class initializer.
 			TableColorizeMode colorizeMode = tableInfo.tableColorizeMode;
-			string colorizeArg = GetArg("colorizeTable","").ToLower();
 			if ( GetBoolArg("autoColorizeTable")) // Legacy support for the -autoColorizeTable arg
 			{
 				colorizeMode = TableColorizeMode.Auto;
 			}
-			if ( colorizeArg != "" )
-			{
-				if (colorizeArg == "auto")
-				{
-					colorizeMode = TableColorizeMode.Auto;
-				}
-				else if (colorizeArg == "off")
-				{
-					colorizeMode = TableColorizeMode.Off;
-				}
-				else if (colorizeArg == "budget")
-				{
-					colorizeMode = TableColorizeMode.Budget;
-				}
-			}
+			colorizeMode = GetEnumArg("colorizeTable", colorizeMode);
+
 
 			bool addMinMaxColumns = !GetBoolArg("noSummaryMinMax") && !bTransposeCollatedSummaryTable;
 
@@ -920,7 +909,9 @@ namespace PerfReportTool
 			SummaryTable filteredTable = table.SortAndFilter(tableInfo.columnFilterList, tableInfo.rowSortList, bReverseTable, weightByColumnName, additionalColumnFilters, bSortTrailingDigitsAsNumeric);
 			if (bCollated)
 			{
-				filteredTable = filteredTable.CollateSortedTable(tableInfo.rowSortList, addMinMaxColumns);
+				StringCollationVisibility stringCollationVisibility = GetEnumArg("collatedStringVisibility", tableInfo.stringCollationVisibility);
+				DateCollationVisibility dateCollationVisibility = GetEnumArg("collatedDateVisibility", tableInfo.dateCollationVisibility);
+				filteredTable = filteredTable.CollateSortedTable(tableInfo.rowSortList, addMinMaxColumns, stringCollationVisibility, dateCollationVisibility);
 			}
 			if (bToCSV)
 			{
diff --git a/Engine/Source/Programs/CSVTools/PerfReportTool/ReportXml.cs b/Engine/Source/Programs/CSVTools/PerfReportTool/ReportXml.cs
--- a/Engine/Source/Programs/CSVTools/PerfReportTool/ReportXml.cs
+++ b/Engine/Source/Programs/CSVTools/PerfReportTool/ReportXml.cs
@@ -958,6 +958,15 @@ namespace PerfReportTool
 				attributeValue = xmlVariableMappings.ResolveVariables(attributeValue);
 			}
 
+			if (typeof(T).IsEnum)
+			{
+#nullable enable
+				if (attributeValue != null && Enum.TryParse(typeof(T), attributeValue, true, out object? result))
+#nullable disable
+				{
+					return (T)result;
+				}
+			}
 			try
 			{
 				switch (Type.GetTypeCode(typeof(T)))
diff --git a/Engine/Source/Programs/CSVTools/PerfReportTool/SummaryTable.cs b/Engine/Source/Programs/CSVTools/PerfReportTool/SummaryTable.cs
--- a/Engine/Source/Programs/CSVTools/PerfReportTool/SummaryTable.cs
+++ b/Engine/Source/Programs/CSVTools/PerfReportTool/SummaryTable.cs
@@ -2,8 +2,8 @@
 
 using System;
 using System.Collections.Generic;
-using System.Data.Common;
 using System.Linq;
+using System.Text;
 using System.Xml.Linq;
 using CSVStats;
 using PerfReportTool;
@@ -100,22 +100,9 @@ namespace PerfSummaries
 			bReverseSortRows = tableElement.GetSafeAttribute<bool>(variableMappings, "reverseSortRows", false);
 			bScrollableFormatting = tableElement.GetSafeAttribute<bool>(variableMappings, "scrollableFormatting", false);
 
-			string colorizeModeStr = tableElement.GetSafeAttribute<string>(variableMappings, "colorizeMode", "").ToLower();
-			if (colorizeModeStr != "")
-			{
-				if (colorizeModeStr == "auto")
-				{
-					tableColorizeMode = TableColorizeMode.Auto;
-				}
-				else if (colorizeModeStr == "off")
-				{
-					tableColorizeMode = TableColorizeMode.Off;
-				}
-				else if (colorizeModeStr == "budget")
-				{
-					tableColorizeMode = TableColorizeMode.Budget;
-				}
-			}
+			tableColorizeMode = tableElement.GetSafeAttribute<TableColorizeMode>(variableMappings, "colorizeMode", tableColorizeMode);
+			dateCollationVisibility = tableElement.GetSafeAttribute<DateCollationVisibility>(variableMappings, "collatedDateVisibility", dateCollationVisibility);
+			stringCollationVisibility = tableElement.GetSafeAttribute<StringCollationVisibility>(variableMappings, "collatedStringVisibility", stringCollationVisibility);
 
 			statThreshold = tableElement.GetSafeAttribute<float>(variableMappings, "statThreshold", 0.0f);
 			hideStatPrefix = tableElement.GetSafeAttribute<string>(variableMappings, "hideStatPrefix");
@@ -191,6 +178,9 @@ namespace PerfSummaries
 		public bool bReverseSortRows;
 		public bool bScrollableFormatting;
 		public TableColorizeMode tableColorizeMode = TableColorizeMode.Budget;
+		public DateCollationVisibility dateCollationVisibility = DateCollationVisibility.Newest;
+		public StringCollationVisibility stringCollationVisibility = StringCollationVisibility.Auto;
+
 		public float statThreshold;
 		public string hideStatPrefix = null;
 		public string weightByColumn = null;
@@ -275,6 +265,21 @@ namespace PerfSummaries
 		All
 	};
 
+	enum DateCollationVisibility
+	{
+		Hide,
+		Newest,
+		Oldest
+	};
+
+	enum StringCollationVisibility
+	{
+		Show,
+		Hide,
+		Auto
+	};
+
+
 	class SummaryTableColumnFormatInfo
 	{
 		public SummaryTableColumnFormatInfo()
@@ -466,11 +471,9 @@ namespace PerfSummaries
 			return "";
 		}
 
-		// Returns true if we should displaya min/max/avg sub columns (if enabled).
-		public bool DisplayAggregates()
+		public bool IsDateFormat()
 		{
-			// Date columns are numeric since they're a timestamp, but we don't want to display them as avg/min/max.
-			return isNumeric && (formatInfo == null || !formatInfo.IsDate());
+			return formatInfo != null && formatInfo.IsDate();
 		}
 
 		private double FilterInvalidValue(double value)
@@ -665,6 +668,19 @@ namespace PerfSummaries
 			return statName;
 		}
 
+		public void ReserveRows(int count)
+		{
+			if (isNumeric)
+			{
+				doubleValues.Capacity = count;
+				colourThresholds.Capacity = count;
+			}
+			else
+			{
+				stringValues.Capacity = count;
+			}			
+		}
+
 		public void SetValue(int index, double value)
 		{
 			if (!isNumeric)
@@ -1228,11 +1244,144 @@ namespace PerfSummaries
 			}
 		}
 
-		public SummaryTable CollateSortedTable(List<string> collateByList, bool addMinMaxColumns)
+		string MakeRowSortKey(List<SummaryTableColumn> sortByColumnList, int rowIndex)
 		{
-			int numSubColumns = addMinMaxColumns ? 3 : 1;
+			StringBuilder sb = new StringBuilder();
+			foreach (SummaryTableColumn column in sortByColumnList)
+			{
+				sb.Append('{');
+				sb.Append(column.GetStringValue(rowIndex));
+				sb.Append('}');
+			}
+			return sb.ToString();
+		}
+
+		private SummaryTableColumn CollateColumn(SummaryTableColumn srcColumn, ColumnAggregateType aggregateType, List<int> collatedRowSourceRowCounts, SummaryTableColumn aggregateBaseColumn = null)
+		{
+			return CollateColumn(srcColumn, aggregateType, collatedRowSourceRowCounts, aggregateBaseColumn, out int temp);
+		}
+
+		private SummaryTableColumn CollateColumn(SummaryTableColumn srcColumn, ColumnAggregateType aggregateType, List<int> collatedRowSourceRowCounts, SummaryTableColumn aggregateBaseColumn, out int validMultiRowCollationsCount)
+		{
+			SummaryTableColumn collatedColumn = new SummaryTableColumn(srcColumn.name, srcColumn.isNumeric, null, false, srcColumn.elementType, srcColumn.formatInfo, srcColumn.tooltip, srcColumn.columnColourModifier, aggregateType, aggregateBaseColumn);
+			collatedColumn.ReserveRows(collatedRowSourceRowCounts.Count);
+
+			validMultiRowCollationsCount = 0;
+			int srcRowIndexOffset = 0;
+			for (int collatedRowIndex = 0; collatedRowIndex < collatedRowSourceRowCounts.Count; collatedRowIndex++)
+			{
+				int mergedRowCount = collatedRowSourceRowCounts[collatedRowIndex];
+				int mergedRowEndIndex = srcRowIndexOffset + mergedRowCount;
+				if (srcColumn.isNumeric)
+				{
+					// Set color thresholds based on the source column
+					for (int srcRowIndex=srcRowIndexOffset; srcRowIndex < mergedRowEndIndex; srcRowIndex++)
+					{
+						collatedColumn.SetColourThresholds(collatedRowIndex, srcColumn.GetColourThresholds(srcRowIndex));
+					}
 
-			// Find all the columns in collateByList
+					switch (aggregateType)
+					{
+						case ColumnAggregateType.Avg:
+							{
+								double total = 0.0;
+								double totalWeight = 0.0;
+								for (int srcRowIndex = srcRowIndexOffset; srcRowIndex < mergedRowEndIndex; srcRowIndex++)
+								{
+									double rowValue = srcColumn.GetValue(srcRowIndex);
+									if (rowValue != double.MaxValue)
+									{
+										double rowWeight = (rowWeightings != null) ? rowWeightings[srcRowIndex] : 1.0;
+										total += rowValue * rowWeight;
+										totalWeight += rowWeight;
+									}
+								}
+								if (totalWeight > 0.0)
+								{
+									collatedColumn.SetValue(collatedRowIndex, total / totalWeight);
+								}
+							}
+							break;
+						case ColumnAggregateType.Min:
+							{
+								double minValue = double.MaxValue;
+								for (int srcRowIndex = srcRowIndexOffset; srcRowIndex < mergedRowEndIndex; srcRowIndex++)
+								{
+									double rowValue = srcColumn.GetValue(srcRowIndex);
+									if (rowValue != double.MaxValue)
+									{
+										minValue = Math.Min(minValue, srcColumn.GetValue(srcRowIndex));
+									}
+								}
+								if (minValue != double.MaxValue)
+								{
+									collatedColumn.SetValue(collatedRowIndex, minValue);
+								}
+							}
+							break;
+						case ColumnAggregateType.Max:
+							{
+								double maxValue = double.MinValue;
+								for (int srcRowIndex = srcRowIndexOffset; srcRowIndex < mergedRowEndIndex; srcRowIndex++)
+								{
+									double rowValue = srcColumn.GetValue(srcRowIndex);
+									if (rowValue != double.MaxValue)
+									{
+										maxValue = Math.Max(maxValue, srcColumn.GetValue(srcRowIndex));
+									}
+								}
+								if (maxValue != double.MinValue)
+								{
+									collatedColumn.SetValue(collatedRowIndex, maxValue);
+								}								
+							}
+							break;
+					}
+					if (mergedRowCount > 1)
+					{
+						validMultiRowCollationsCount++;
+					}					
+				}
+				else 
+				{
+					// Collate non-numeric values. Display if they're the same, otherwise show "(multiple)"
+					string mergedValue = null;
+					bool bStringsMatch = true;
+					for (int srcRowIndex = srcRowIndexOffset; srcRowIndex < mergedRowEndIndex; srcRowIndex++)
+					{
+						string currentValue = srcColumn.GetStringValue(srcRowIndex);
+						if (mergedValue == null)
+						{
+							mergedValue = currentValue;
+						}
+						else
+						{
+							if (currentValue != mergedValue)
+							{
+								// The strings differ so we can't merge them
+								mergedValue = "<i>(multiple)</i>"; 
+								bStringsMatch = false;
+								break;
+							}
+						}
+					}
+
+					// If the rows match and are not null or empty then count them.
+					// Note that we do not count as valid merged rows unless Count>1 since otherwise a single count of 1 would cause a row to be visible if CollatedStringVisibility is Auto
+					if (mergedValue != null && mergedValue != "" && bStringsMatch && mergedRowCount > 1)
+					{
+						validMultiRowCollationsCount++;
+					}
+					collatedColumn.SetStringValue(collatedRowIndex, mergedValue == null ? "" : mergedValue);
+				}
+				srcRowIndexOffset += mergedRowCount;
+			}
+			return collatedColumn;
+		}
+
+		public SummaryTable CollateSortedTable(List<string> collateByList, bool addMinMaxColumns, StringCollationVisibility stringCollationVisibility, DateCollationVisibility dateCollationVisibility)
+		{
+			// Find the collateBy columns
 			HashSet<SummaryTableColumn> collateByColumns = new HashSet<SummaryTableColumn>();
 			foreach (string collateBy in collateByList)
 			{
@@ -1247,166 +1396,110 @@ namespace PerfSummaries
 				throw new Exception("None of the metadata strings were found:" + string.Join(", ", collateByList));
 			}
 
-			// Add the new collateBy columns in the order they appear in the original column list
-			List<SummaryTableColumn> newColumns = new List<SummaryTableColumn>();
-			List<string> finalSortByList = new List<string>();
+			// Generate the filtered sortBy column list
+			List<SummaryTableColumn> sortByColumnList = new List<SummaryTableColumn>();
 			foreach (SummaryTableColumn srcColumn in columns)
 			{
 				if (collateByColumns.Contains(srcColumn))
 				{
-					newColumns.Add(new SummaryTableColumn(srcColumn.name, false, srcColumn.displayName, false, srcColumn.elementType, srcColumn.formatInfo));
-					finalSortByList.Add(srcColumn.name.ToLower());
+					// finalSortByList.Add(srcColumn.name.ToLower());
+					sortByColumnList.Add(srcColumn);
 					// Early out if we've found all the columns
-					if (finalSortByList.Count == collateByColumns.Count)
+					if (sortByColumnList.Count == collateByColumns.Count)
 					{
 						break;
 					}
 				}
 			}
 
-			newColumns.Add(new SummaryTableColumn("Count", true, null, false, SummaryTableElement.Type.ToolMetadata, bInIsCountColumn:true));
-			int countColumnIndex = newColumns.Count - 1;
-
-			int numericColumnStartIndex = newColumns.Count;
-			List<int> srcToDestBaseColumnIndex = new List<int>();
-			foreach (SummaryTableColumn column in columns)
+			// Count the number of rows for each collated row
+			List<int> collatedRowSourceRowCounts = new List<int>();
+			int sourceRowCount = 0;
+			string currentRowSortKey = MakeRowSortKey(sortByColumnList, 0);
+			for (int i = 0; i < rowCount; i++)
 			{
-				// Add avg/min/max columns for this column if it's numeric and we didn't already add it above 
-				if (column.DisplayAggregates() && !collateByColumns.Contains(column))
-				{
-					srcToDestBaseColumnIndex.Add(newColumns.Count);
-					SummaryTableColumn avgColumn = new SummaryTableColumn(column.name, true, null, false, column.elementType, column.formatInfo, column.tooltip, column.columnColourModifier, ColumnAggregateType.Avg);
-					newColumns.Add(avgColumn);
-					if (addMinMaxColumns)
-					{
-						newColumns.Add(new SummaryTableColumn(column.name, true, null, false, column.elementType, column.formatInfo, column.tooltip, column.columnColourModifier, ColumnAggregateType.Min, avgColumn));
-						newColumns.Add(new SummaryTableColumn(column.name, true, null, false, column.elementType, column.formatInfo, column.tooltip, column.columnColourModifier, ColumnAggregateType.Max, avgColumn));
-					}
-				}
-				else
+				string nextSortKey = (i + 1 < rowCount) ? MakeRowSortKey(sortByColumnList, i + 1) : null;
+				sourceRowCount++;
+
+				// If this is the last row (nextSortKey=null) or if the sort key is different then write out the count
+				if (nextSortKey != currentRowSortKey)
 				{
-					srcToDestBaseColumnIndex.Add(-1);
+					collatedRowSourceRowCounts.Add(sourceRowCount);
+					currentRowSortKey = nextSortKey;
+					sourceRowCount = 0;
 				}
 			}
 
-			List<double> RowMaxValues = new List<double>();
-			List<double> RowTotals = new List<double>();
-			List<double> RowMinValues = new List<double>();
-			List<int> RowCounts = new List<int>();
-			List<double> RowWeights = new List<double>();
-			List<ColourThresholdList> RowColourThresholds = new List<ColourThresholdList>();
+			List<SummaryTableColumn> newColumns = new List<SummaryTableColumn>();
 
-			// Set the initial sort key
-			string CurrentRowSortKey = "";
-			foreach (string collateBy in finalSortByList)
+			// Generate the collated sortby columns
+			foreach (SummaryTableColumn srcColumn in sortByColumnList)
 			{
-				CurrentRowSortKey += "{" + columnLookup.Get(collateBy).GetStringValue(0) + "}";
+				SummaryTableColumn newColumn = new SummaryTableColumn(srcColumn.name, false, srcColumn.displayName, false, srcColumn.elementType, srcColumn.formatInfo);
+				newColumns.Add(newColumn);
+				int srcRowIndex = 0;
+				for (int i=0; i<collatedRowSourceRowCounts.Count; i++)
+				{
+					newColumn.SetStringValue(i, srcColumn.GetStringValue(srcRowIndex));
+					srcRowIndex += collatedRowSourceRowCounts[i];
+				}				
 			}
 
-			int destRowIndex = 0;
-			bool reset = true;
-			int mergedRowsCount = 0;
-			for (int i = 0; i < rowCount; i++)
+			// Add the count column
+			SummaryTableColumn countColumn = new SummaryTableColumn("Count", true, null, false, SummaryTableElement.Type.ToolMetadata, bInIsCountColumn: true);
+			for (int i=0; i< collatedRowSourceRowCounts.Count; i++)
+			{
+				countColumn.SetValue(i, collatedRowSourceRowCounts[i]);
+			}
+			newColumns.Add(countColumn);
+
+			int nonCollatedStartIndex = newColumns.Count;
+			foreach (SummaryTableColumn column in columns)
 			{
-				if (reset)
+				// Add avg/min/max columns for this column if it's numeric and we didn't already add it above 
+				if (!collateByColumns.Contains(column))
 				{
-					RowMaxValues.Clear();
-					RowMinValues.Clear();
-					RowTotals.Clear();
-					RowCounts.Clear();
-					RowWeights.Clear();
-					RowColourThresholds.Clear();
-					for (int j = 0; j < columns.Count; j++)
+					if (column.isNumeric)
 					{
-						if (addMinMaxColumns)
+						if (column.IsDateFormat())
 						{
-							RowMaxValues.Add(-double.MaxValue);
-							RowMinValues.Add(double.MaxValue);
+							if (dateCollationVisibility != DateCollationVisibility.Hide)
+							{
+								// If we're showing collated dates then determine min/max based on the visibility mode
+								// TODO: expose this to format so we can set per column
+								newColumns.Add(CollateColumn(column, dateCollationVisibility == DateCollationVisibility.Newest ? ColumnAggregateType.Max : ColumnAggregateType.Min, collatedRowSourceRowCounts));
+							}
 						}
-						RowTotals.Add(0.0f);
-						RowCounts.Add(0);
-						RowWeights.Add(0.0);
-						RowColourThresholds.Add(null);
-					}
-					mergedRowsCount = 0;
-					reset = false;
-				}
-
-				// Compute min/max/total for all numeric columns
-				for (int j = 0; j < columns.Count; j++)
-				{
-					SummaryTableColumn column = columns[j];
-					if (column.DisplayAggregates())
-					{
-						double value = column.GetValue(i);
-						if (value != double.MaxValue)
+						else
 						{
+							// Generate aggregate min/max/avg values
+							SummaryTableColumn avgColumn = CollateColumn(column, ColumnAggregateType.Avg, collatedRowSourceRowCounts);
+							newColumns.Add(avgColumn);
 							if (addMinMaxColumns)
 							{
-								RowMaxValues[j] = Math.Max(RowMaxValues[j], value);
-								RowMinValues[j] = Math.Min(RowMinValues[j], value);
+								newColumns.Add(CollateColumn(column, ColumnAggregateType.Min, collatedRowSourceRowCounts, avgColumn));
+								newColumns.Add(CollateColumn(column, ColumnAggregateType.Max, collatedRowSourceRowCounts, avgColumn));
 							}
-							RowColourThresholds[j] = column.GetColourThresholds(i);
-							RowCounts[j]++;
-							double rowWeight = (rowWeightings != null) ? rowWeightings[i] : 1.0;
-							RowWeights[j] += rowWeight;
-							RowTotals[j] += value * rowWeight;
 						}
 					}
-				}
-				mergedRowsCount++;
-
-				// Are we done?
-				string nextSortKey = "";
-				if (i < rowCount - 1)
-				{
-					foreach (string collateBy in finalSortByList)
-					{
-						nextSortKey += "{" + columnLookup.Get(collateBy).GetStringValue(i + 1) + "}";
-					}
-				}
-
-				// If this is the last row or if the sort key is different then write it out
-				if (nextSortKey != CurrentRowSortKey)
-				{
-					for (int j = 0; j < countColumnIndex; j++)
-					{
-						string key = newColumns[j].name.ToLower();
-						newColumns[j].SetStringValue(destRowIndex, columnLookup.Get(key).GetStringValue(i));
-					}
-					// Commit the row 
-					newColumns[countColumnIndex].SetValue(destRowIndex, (double)mergedRowsCount);
-					for (int j = 0; j < columns.Count; j++)
+					else if (stringCollationVisibility != StringCollationVisibility.Hide)
 					{
-						int destColumnBaseIndex = srcToDestBaseColumnIndex[j];
-						if (destColumnBaseIndex != -1 && RowCounts[j] > 0)
+						SummaryTableColumn stringDataColumn = CollateColumn(column, ColumnAggregateType.None, collatedRowSourceRowCounts, null, out int validMultiRowCollationsCount);
+						// If StringCollationVisibility is Auto, only show the column if there were successful valid multi-row collations (ie ignoring collations where row count is 1)
+						if ( validMultiRowCollationsCount > 0 || stringCollationVisibility == StringCollationVisibility.Show )
 						{
-							newColumns[destColumnBaseIndex].SetValue(destRowIndex, RowTotals[j] / RowWeights[j]);
-							if (addMinMaxColumns && newColumns[destColumnBaseIndex].DisplayAggregates())
-							{
-								newColumns[destColumnBaseIndex + 1].SetValue(destRowIndex, RowMinValues[j]);
-								newColumns[destColumnBaseIndex + 2].SetValue(destRowIndex, RowMaxValues[j]);
-							}
-
-							// Set colour thresholds based on the source column
-							ColourThresholdList Thresholds = RowColourThresholds[j];
-							for (int k = 0; k < numSubColumns; k++)
-							{
-								newColumns[destColumnBaseIndex + k].SetColourThresholds(destRowIndex, Thresholds);
-							}
+							newColumns.Add(stringDataColumn);
 						}
 					}
-					reset = true;
-					destRowIndex++;
 				}
-				CurrentRowSortKey = nextSortKey;
 			}
 
+
 			SummaryTable newTable = new SummaryTable();
 			newTable.columns = newColumns;
 			newTable.InitColumnLookup();
-			newTable.rowCount = destRowIndex;
-			newTable.firstStatColumnIndex = numericColumnStartIndex;
+			newTable.rowCount = collatedRowSourceRowCounts.Count;
+			newTable.firstStatColumnIndex = nonCollatedStartIndex; 
 			newTable.isCollated = true;
 			newTable.hasMinMaxColumns = addMinMaxColumns;
 			return newTable;
@@ -1811,11 +1904,11 @@ namespace PerfSummaries
 			bool showFilteredColumnDebug = false)
 		{
 			System.IO.StreamWriter htmlFile = new System.IO.StreamWriter(htmlFilename, false);
-			int statColSpan = hasMinMaxColumns ? 3 : 1;
-			int cellPadding = 2;
+
+			string cellPadding = "2px 4px 2px 4px";
 			if (isCollated)
 			{
-				cellPadding = 4;
+				cellPadding = "4px 4px 4px 4px";
 			}
 
 			// Generate an automatic title
@@ -1889,10 +1982,10 @@ namespace PerfSummaries
 					"  for (var i=0;i<numStickyCols;i++) \n" +
 					"  { \n" +
 					"	 var rBorderParam=(i==numStickyCols-1) ? 'border-right: 2px solid black;':''; \n" +
-					"	 styleElement.textContent+='tr.lastHeaderRow th:nth-child('+(i+1)+') {  z-index: 8;  border-top: 2px solid black;  font-size: 11px;  left: '+xOffset+'px;'+rBorderParam+'}'; \n" +
-					"	 styleElement.textContent+='td:nth-child('+(i+1)+') {  position: -webkit-sticky;  position: sticky; z-index: 7;  left: '+xOffset+'px; '+rBorderParam+'}'; \n" +
+					"	 styleElement.textContent+='tr.lastHeaderRow th:nth-child('+(i+1)+') {  z-index: 8;  border-top: 2px solid black;  font-size: 11px; left:'+xOffset+'px;'+rBorderParam+'}'; \n" +
+					"	 styleElement.textContent+='td:nth-child('+(i+1)+') {  position: -webkit-sticky;  position: sticky; z-index: 7; left:'+xOffset+'px; '+rBorderParam+'}'; \n" +
 					"	 xOffset+=table.rows[1].cells[i].offsetWidth; \n" +
-					"  } \n"+
+					"  } \n" +
 					"} \n"
 					);
 				htmlFile.WriteLine("</script>");
@@ -1922,7 +2015,7 @@ namespace PerfSummaries
 
 				tableCss =
 					"table {table-layout: fixed;} \n" +
-					"table, th, td { border: 0px solid black; border-spacing: 0; border-collapse: separate; padding: " + cellPadding + "px; vertical-align: center; font-family: 'Verdana', Times, serif; font-size: " + cellFontSize + "px;} \n" +
+					"table, th, td { border: 0px solid black; border-spacing: 0; border-collapse: separate; padding: " + cellPadding + "; vertical-align: center; font-family: 'Verdana', Times, serif; font-size: " + cellFontSize + "px;} \n" +
 					"td {" +
 					"  border-right: 1px solid black;" +
 					"  max-width: 450;" +
@@ -1980,7 +2073,7 @@ namespace PerfSummaries
 			else
 			{
 				tableCss =
-					"table, th, td { border: 2px solid black; border-collapse: collapse; padding: " + cellPadding + "px; vertical-align: center; font-family: 'Verdana', Times, serif; font-size: 11px;} \n";
+					"table, th, td { border: 2px solid black; border-collapse: collapse; padding: " + cellPadding + "; vertical-align: center; font-family: 'Verdana', Times, serif; font-size: 11px;} \n";
 			}
 
 
@@ -2035,16 +2128,36 @@ namespace PerfSummaries
 						topHeaderRow.AddCell("", "colspan='" + firstStatColumnIndex + "'");
 					}
 					// Add the stat columns
+					int aggregateColumnCount = 0;
 					for (int i = firstStatColumnIndex; i < columns.Count; i++)
 					{
 						SummaryTableColumn column = columns[i];
 						string statName = GetBaseStatNameWithPrefixAndSuffix(column.GetDisplayName(hideStatPrefix, bAddStatNameSpacing, bGreyOutStatPrefixes), out string prefix, out string suffix);
-						if ((i - 1) % statColSpan == 0)
+						string cellText = prefix.Trim();
+						if (column.aggregateType == ColumnAggregateType.None)
 						{
-							string attributes = $"colspan='{statColSpan}'{String.Join(" ", column.GetHeaderAttributes())}";
+							cellText = "-";
+						}
+						else if (column.IsDateFormat())
+						{
+							cellText = column.aggregateType == ColumnAggregateType.Max ? "<i>newest</i>" : "<i>oldest</i>";
+						}
+						headerRow.AddCell(cellText);
+						aggregateColumnCount++;
+
+						// Update the top row if this is the last column of an aggregate (ie max)
+						if (column.aggregateType == ColumnAggregateType.None || column.aggregateBaseColumn == null || column.aggregateType == ColumnAggregateType.Max)
+						{
+							string attributes = "";
+							if (aggregateColumnCount > 1)
+							{
+								attributes = $"colspan='{aggregateColumnCount}' ";
+							}
+							attributes += String.Join(" ", column.GetHeaderAttributes());
 							topHeaderRow.AddCell(statName + suffix, attributes);
+							aggregateColumnCount = 0;
 						}
-						headerRow.AddCell(prefix.Trim());
+
 					}
 				}
 				else
