PerfReportTool 4.244.1: Fix crash due to date column diff handling when min/max is enabled.

diff --git a/Engine/Source/Programs/CSVTools/PerfReportTool/PerfReportTool.cs b/Engine/Source/Programs/CSVTools/PerfReportTool/PerfReportTool.cs
--- a/Engine/Source/Programs/CSVTools/PerfReportTool/PerfReportTool.cs
+++ b/Engine/Source/Programs/CSVTools/PerfReportTool/PerfReportTool.cs
@@ -23,7 +23,7 @@ namespace PerfReportTool
     class Version
     {
 		// Format: Major.Minor.Bugfix
-        private static string VersionString = "4.244.0";
+        private static string VersionString = "4.244.1";
 
         public static string Get() { return VersionString; }
     };
diff --git a/Engine/Source/Programs/CSVTools/PerfReportTool/SummaryTable.cs b/Engine/Source/Programs/CSVTools/PerfReportTool/SummaryTable.cs
--- a/Engine/Source/Programs/CSVTools/PerfReportTool/SummaryTable.cs
+++ b/Engine/Source/Programs/CSVTools/PerfReportTool/SummaryTable.cs
@@ -504,7 +504,7 @@ namespace PerfSummaries
 			List<string> newToolTips = new List<string>(toolTips.Count > 0 ? newCount : 0);
 			List<ColourThresholdList> newColourThresholds = new List<ColourThresholdList>(colourThresholds.Count > 0 ? newCount : 0);
 
-			bool bComputeDiff = isNumeric && !isCountColumn;
+			bool bComputeDiff = isNumeric && !IsDateFormat() && !isCountColumn;
 
 			static bool NeedsDiffColumn(int originalRowIndex, bool bShowIntermediateDiffRows)
 			{
@@ -1873,7 +1873,7 @@ namespace PerfSummaries
 				if (column.isNumeric && !column.isCountColumn)
 				{
 					// If this is a min or max column then we use the Avg column to compute the diff score. All 3 columns will be filtered together
-					double maxAbsDiff = hasMinMaxColumns ? column.aggregateBaseColumn.GetMaxAbsDiff() : column.GetMaxAbsDiff();
+					double maxAbsDiff = ( hasMinMaxColumns && column.aggregateBaseColumn != null ) ? column.aggregateBaseColumn.GetMaxAbsDiff() : column.GetMaxAbsDiff();
 					if (maxAbsDiff >= threshold)
 					{
 						newColumns.Add(column);
