Add support for csv.bin files to CSVToTelemetry
This allows us to uplaod a CSVProfile binary or text file to multiple telemetry backends directly via StudioTelemetry
Usage : -csvprofile=myfile.csv.bin or -csvprofile=myfile.csv

diff --git a/Engine/Source/Programs/CSVTools/CSVToTelemetry/CSVToTelemetry.Build.cs b/Engine/Source/Programs/CSVTools/CSVToTelemetry/CSVToTelemetry.Build.cs
--- a/Engine/Source/Programs/CSVTools/CSVToTelemetry/CSVToTelemetry.Build.cs
+++ b/Engine/Source/Programs/CSVTools/CSVToTelemetry/CSVToTelemetry.Build.cs
@@ -15,6 +15,7 @@ public class CSVToTelemetry : ModuleRules
 			"BuildSettings",
 			"Core",
 			"CoreUObject",
+			"CSVUtils",
 			"Horde",
 			"HTTP",
 			"Projects",
diff --git a/Engine/Source/Programs/CSVTools/CSVToTelemetry/Private/CSVToTelemetry.cpp b/Engine/Source/Programs/CSVTools/CSVToTelemetry/Private/CSVToTelemetry.cpp
--- a/Engine/Source/Programs/CSVTools/CSVToTelemetry/Private/CSVToTelemetry.cpp
+++ b/Engine/Source/Programs/CSVTools/CSVToTelemetry/Private/CSVToTelemetry.cpp
@@ -1,6 +1,7 @@
 // Copyright Epic Games, Inc. All Rights Reserved.
 
 #include "CSVToTelemetry.h"
+#include "CSVProfilerUtils.h"
 #include "RequiredProgramMainCPPInclude.h"
 #include "Analytics.h"
 #include "AnalyticsET.h"
@@ -12,34 +13,19 @@
 DEFINE_LOG_CATEGORY_STATIC(LogCSVToTelemetry, Log, All);
 IMPLEMENT_APPLICATION(CSVToTelemetry, "CSVToTelemetry");
 
-
-
-INT32_MAIN_INT32_ARGC_TCHAR_ARGV()
+static int GenerateTelemetryFromCSVFile(const FString& FilePath)
 {
-	FTaskTagScope Scope(ETaskTag::EGameThread);
-
-	// Allows this program to accept a project argument on the commandline and use project-specific config
-	UE::ProjectUtilities::ParseProjectDirFromCommandline(ArgC, ArgV);
-
-	// start up the main loop,
-	GEngineLoop.PreInit(ArgC, ArgV);
-
 	int32 Result = 1;
 
-	FString FilePath;
 	FString EventName;
 
-	if (FParse::Value(FCommandLine::Get(), TEXT("csv="), FilePath)==false)
-	{
-		UE_LOG(LogCSVToTelemetry, Error, TEXT("Must provide a CSV file with -csv=[fillename]"));
-	}
-
 	if (FParse::Value(FCommandLine::Get(), TEXT("event="), EventName) == false)
 	{
 		UE_LOG(LogCSVToTelemetry, Error, TEXT("Must provide a Event name with -event=[eventname]"));
+		return Result;
 	}
-	
-	if (FilePath.Len()>0 && EventName.Len()>0)
+
+	if (FilePath.Len() > 0)
 	{
 		TArray<FString> Rows;
 
@@ -62,9 +48,18 @@ INT32_MAIN_INT32_ARGC_TCHAR_ARGV()
 
 				FString Columns;
 
-				uint32 SchemaVersion(0);
-				FParse::Value(FCommandLine::Get(), TEXT("schema="), SchemaVersion);
-				
+				int32 SchemaVersion(-1);
+
+				TSharedPtr<IAnalyticsProvider> Provider = FStudioTelemetry::Get().GetProvider().Pin();
+				TArray<FAnalyticsEventAttribute> DefaultAttributes = Provider->GetDefaultEventAttributesSafe();
+
+				if (FParse::Value(FCommandLine::Get(), TEXT("schema="), SchemaVersion) == true)
+				{
+					DefaultAttributes.Emplace(TEXT("SchemaVersion"), SchemaVersion);
+				}
+
+				Provider->SetDefaultEventAttributes(MoveTemp(DefaultAttributes));
+
 				if (FParse::Value(FCommandLine::Get(), TEXT("columns="), Columns) == true)
 				{
 					// We have specified the CSV columns on the command line 
@@ -98,19 +93,13 @@ INT32_MAIN_INT32_ARGC_TCHAR_ARGV()
 
 							TArray<FAnalyticsEventAttribute> Attributes;
 
-							// If we've specified a non-zero schema version then add that
-							if (SchemaVersion != 0)
-							{
-								Attributes.Emplace(TEXT("SchemaVersion"), SchemaVersion);
-							}
-
 							// Add the Key/Value pair to the attribute list
 							for (int32 i = 0; i < ValueArray.Num(); ++i)
 							{
 								ValueArray[i].RemoveSpacesInline();
 								Attributes.Emplace(KeyArray[i], ValueArray[i]);
 							}
-	
+
 							// Keep track of the total data we're uploading
 							TotalUploadSize += KeySize + ValueSize;
 
@@ -140,12 +129,121 @@ INT32_MAIN_INT32_ARGC_TCHAR_ARGV()
 		}
 	}
 
+	return Result;
+}
+
+static int GenerateTelemetryFromCSVProfileFile(const FString& FilePath)
+{
+	int32 Result = 1;
+
+	FString EventName;
+
+	if (FParse::Value(FCommandLine::Get(), TEXT("event="), EventName) == false)
+	{
+		UE_LOG(LogCSVToTelemetry, Error, TEXT("Must provide a Event name with -event=[eventname]"));
+		return Result;
+	}
+
+	CsvUtils::FCsvProfilerCapture Capture;
+
+	if (FilePath.Len() > 0)
+	{
+		if (FilePath.Contains(".bin"))
+		{
+			// Binary CSV file
+			CsvUtils::ReadFromCsvBin(Capture, *FilePath);
+		}
+		else
+		{
+			// Text CSV file
+			CsvUtils::ReadFromCsv(Capture, *FilePath);
+		}
+	}
+
+	if (Capture.Events.Num())
+	{
+		FStudioTelemetry::Get().StartSession();
+
+		if (FStudioTelemetry::Get().IsSessionRunning())
+		{
+			TSharedPtr<IAnalyticsProvider> Provider = FStudioTelemetry::Get().GetProvider().Pin();
+			TArray<FAnalyticsEventAttribute> DefaultAttributes = Provider->GetDefaultEventAttributesSafe();
+
+			for (TMap<FString, FString>::TConstIterator It(Capture.Metadata); It; ++It)
+			{
+				DefaultAttributes.Emplace(It->Key, It->Value);
+			}
+
+			Provider->SetDefaultEventAttributes(MoveTemp(DefaultAttributes));
+
+			for (const CsvUtils::FCsvProfilerEvent& Event : Capture.Events)
+			{
+				int32 Frame = Event.Frame;
+
+				TArray<FAnalyticsEventAttribute> Attributes;
+
+				Attributes.Emplace(TEXT("Name"), Event.Name);
+
+				for (TMap<FString, CsvUtils::FCsvProfilerSample>::TConstIterator It(Capture.Samples); It; ++It)
+				{
+					const FString& Name = It->Key;
+					const CsvUtils::FCsvProfilerSample& Sample = It->Value;
+
+					if (Frame < Sample.Values.Num())
+					{
+						Attributes.Emplace(Name, Sample.Values[Frame]);
+					}
+				}
+
+				FStudioTelemetry::Get().RecordEvent(EventName, Attributes);
+			}
+
+			FStudioTelemetry::Get().FlushEvents();
+			FStudioTelemetry::Get().EndSession();
+
+			Result = 0;
+		}
+	}
+
+	return Result;
+}
+
+INT32_MAIN_INT32_ARGC_TCHAR_ARGV()
+{
+	FTaskTagScope Scope(ETaskTag::EGameThread);
+
+	FDateTime StartTime = FDateTime::UtcNow();
+
+	// Allows this program to accept a project argument on the command-line and use project-specific config
+	UE::ProjectUtilities::ParseProjectDirFromCommandline(ArgC, ArgV);
+
+	// start up the main loop,
+	GEngineLoop.PreInit(ArgC, ArgV);
+
+	int32 Result = 1;
+
+	FString FilePath;
+	
+	if (FParse::Value(FCommandLine::Get(), TEXT("csvprofile="), FilePath) == true)
+	{
+		Result = GenerateTelemetryFromCSVProfileFile(FilePath);
+	}
+	else if (FParse::Value(FCommandLine::Get(), TEXT("csv="), FilePath)==true)
+	{
+		Result = GenerateTelemetryFromCSVFile(FilePath);
+	}
+
 	if (Result!=0)
 	{
 		// Always show the usage example if we have not been successful
 		UE_LOG(LogCSVToTelemetry, Display, TEXT("Usage CSVToTelemetry.exe -csv=[filename] -event=[eventname] ( -schema=[value] -columns=[name1|name2|....] )"));
 	}
-
+	else
+	{
+		// Upload completed successfully
+		UE_LOG(LogCSVToTelemetry, Display, TEXT("CSVToTelemetry upload succeeded in %0.2f seconds"), (FDateTime::UtcNow() - StartTime).GetTotalSeconds());
+	}
+	
 	if (FParse::Param(FCommandLine::Get(), TEXT("fastexit")))
 	{
 		FPlatformMisc::RequestExitWithStatus(true, Result);
diff --git a/Engine/Source/Programs/CSVTools/CSVUtils/CSVUtils.build.cs b/Engine/Source/Programs/CSVTools/CSVUtils/CSVUtils.build.cs
new file mode 100644
--- /dev/null
+++ b/Engine/Source/Programs/CSVTools/CSVUtils/CSVUtils.build.cs
@@ -0,0 +1,19 @@
+// Copyright Epic Games, Inc. All Rights Reserved.
+
+using System.IO;
+using UnrealBuildTool;
+
+public class CSVUtils : ModuleRules
+{
+	public CSVUtils(ReadOnlyTargetRules Target) : base(Target)
+	{
+		PCHUsage = ModuleRules.PCHUsageMode.UseExplicitOrSharedPCHs;
+		
+		PrivateDependencyModuleNames.AddRange(
+			new string[] {
+				"Core",
+				"CoreUObject"
+			}
+		);
+	}
+}
diff --git a/Engine/Source/Programs/CSVTools/CSVUtils/Private/CSVProfilerUtils.cpp b/Engine/Source/Programs/CSVTools/CSVUtils/Private/CSVProfilerUtils.cpp
new file mode 100644
--- /dev/null
+++ b/Engine/Source/Programs/CSVTools/CSVUtils/Private/CSVProfilerUtils.cpp
@@ -0,0 +1,482 @@
+// Copyright Epic Games, Inc. All Rights Reserved.
+
+#include "CSVProfilerUtils.h"
+
+#include "HAL/FileManager.h"
+#include "Misc/Compression.h"
+#include "Misc/FileHelper.h"
+#include "Misc/Paths.h"
+#include "Misc/ScopedSlowTask.h"
+#include "String/LexFromString.h"
+#include "String/ParseTokens.h"
+#include "Internationalization/Internationalization.h"
+#include "Internationalization/Text.h"
+
+#define LOCTEXT_NAMESPACE "CsvUtils"
+
+namespace CsvUtils
+{
+
+namespace CsvUtils::Private
+{
+
+static const FText IncorrectFormatText(
+	LOCTEXT("IncorrectFormat", "Incorrect file format - couldn't read expected magic."));
+static const FText UnsupportedVersionText(LOCTEXT("UnsupportedVersion", "File is of an unsupported version."));
+static const FText UnsupportedCompressionTypeText(
+	LOCTEXT("UnsupportedCompressionType", "File uses an unsupported compression type."));
+static const FText UncompressedFormatSupportText(
+	LOCTEXT("UncompressedFormatSupport", "Uncompressed format loading is not yet supported."));
+static const FText SampleDataNotFoundText(
+	LOCTEXT("SampleDataNotFound", "Unable to find sample with name {SampleName} to serialize."));
+static const FText SourceDataNotFoundText(
+	LOCTEXT("SourceDataNotFound", "Unable to find data with name {SourceName} to populate sample."));
+static const FText FailedToConvertValueText(
+	LOCTEXT("FailedToConvertValue", "Failed to convert value {Value} from {SourceUnit} to {DestUnit}."));
+
+static constexpr FStringView SampleNameKey(TEXTVIEW("SampleName"));
+static constexpr FStringView SourceNameKey(TEXTVIEW("SourceName"));
+static constexpr FStringView ValueKey(TEXTVIEW("Value"));
+static constexpr FStringView SourceUnitKey(TEXTVIEW("SourceUnit"));
+static constexpr FStringView DestUnitKey(TEXTVIEW("DestUnit"));
+
+// Should be kept up-to-date with CSVStats.CsvBinVersion
+enum class ECsvBinVersion : int32
+{
+	PreRelease = 1,
+	InitialRelease,
+	CompressionSupportAndFlags,
+
+	COUNT,
+	CURRENT = COUNT - 1
+};
+
+// Should be kept up-to-date with CSVStats.CsvBinFlags
+enum class ECsvBinFlags : uint32
+{
+	None = 0,
+	HasMetadata = 0x00000001,
+};
+ENUM_CLASS_FLAGS(ECsvBinFlags);
+
+// Should be kept up-to-date with CSVStats.CsvBinCompressionType
+enum class CsvBinCompressionType : uint8
+{
+	MsDeflate
+};
+
+// Should be kept up-to-date with CSVStats.ECsvBinCompressionLevel
+enum class ECsvBinCompressionLevel : uint8
+{
+	None,
+	Min,
+	Max
+};
+
+struct FCsvProfilerFileReader
+{
+	explicit FCsvProfilerFileReader(TArray<FText>* InOutErrors) : OutErrors(InOutErrors)
+	{
+	}
+
+	TArray<FText>* OutErrors = nullptr;
+
+	int32 FrameNumber = INDEX_NONE;
+	TArray<FString> SampleNames;
+	TArray<TArray<float>> SampleData;
+	TArray<FCsvProfilerEvent> Events;
+
+	void operator()(FStringView Line)
+	{
+		TRACE_CPUPROFILER_EVENT_SCOPE(Private::FCsvFileReader::operator());
+
+		// Metadata line
+		if (UNLIKELY(Line.StartsWith(TEXT('['))))
+		{
+			// unimplemented
+		}
+		else
+		{
+			// Process a CSV line. Read each token in between commas.
+			// Note that this doesn't properly handle commas in quotes, but we don't
+			// expect any in the data we're processing.
+			bool bHeadingsRow = false;
+			int32 ColumnIndex = INDEX_NONE;
+			ParseTokens(
+				Line, TEXT(','),
+				[&](FStringView Token)
+				{
+					++ColumnIndex;
+					// The first column is known - it contains events. So we can check if this is a headings row
+					// by checking if this column contains this label.
+					if (ColumnIndex == 0)
+					{
+						bHeadingsRow = Token == TEXTVIEW("EVENTS");
+						if (!bHeadingsRow)
+						{
+							// Bump the frame number for each non-headings row we come across.
+							++FrameNumber;
+							//  Split events on the delimiter.
+							ParseTokens(
+								Token, TEXT(';'),
+								[&](FStringView EventToken)
+								{ Events.Add({.Name = FString(EventToken), .Frame = FrameNumber}); },
+								UE::String::EParseTokensOptions::IgnoreCase | UE::String::EParseTokensOptions::SkipEmpty);
+						}
+						return;
+					}
+					const int32 SampleIndex = ColumnIndex - 1;
+					if (UNLIKELY(bHeadingsRow))
+					{
+						SampleNames.SetNum(FMath::Max(SampleNames.Num(), SampleIndex + 1));
+						SampleNames[SampleIndex] = Token;
+					}
+					else
+					{
+						SampleData.SetNum(FMath::Max(SampleData.Num(), ColumnIndex));
+						SampleData[SampleIndex].SetNumZeroed(FrameNumber + 1);
+						LexFromString(SampleData[SampleIndex][FrameNumber], Token);
+					}
+				},
+				UE::String::EParseTokensOptions::IgnoreCase);
+		}
+	}
+
+	bool Finalize(FCsvProfilerCapture& OutCapture)
+	{
+		TRACE_CPUPROFILER_EVENT_SCOPE(Private::FCsvFileReader::Finalize);
+		OutCapture.Samples.Reserve(SampleNames.Num());
+		for (int32 SeriesIndex = 0; SeriesIndex < SampleNames.Num(); ++SeriesIndex)
+		{
+			FCsvProfilerSample& Sample = OutCapture.Samples.Add(SampleNames[SeriesIndex]);
+			Sample.Name = SampleNames[SeriesIndex];
+			Sample.Values = MoveTemp(SampleData[SeriesIndex]);
+		}
+		OutCapture.Events = MoveTemp(Events);
+		return true;
+	}
+};
+
+/**
+ * Helper function to decode a 7-bit encoded integer from the parameter archive.
+ * @param Ar The archive to read from.
+ * @return The decoded integer.
+ */
+static uint64 Decode7bit(FArchive& Ar)
+{
+	uint64 Value = 0;
+	uint64 ByteIndex = 0;
+	bool HasMoreBytes;
+	do
+	{
+		uint8 ByteValue;
+		Ar.Serialize(&ByteValue, 1);
+		HasMoreBytes = ByteValue & 0x80;
+		Value |= uint64(ByteValue & 0x7f) << (ByteIndex * 7);
+		++ByteIndex;
+	} while (HasMoreBytes);
+	return Value;
+}
+
+/**
+ * Helper function to read a C# BinaryWriter-serialized string into a string builder.
+ * @tparam CharType Character type.
+ * @param Ar Binary archive to read the string from.
+ * @param Builder String builder to write to.
+ */
+template <typename CharType>
+void SerializeCsString(FArchive& Ar, TStringBuilderBase<CharType>& Builder)
+{
+	// C# BinaryWriter prefixes strings with their length as a 7-bit encoded int
+	uint32 StringLength = Decode7bit(Ar);
+
+	if constexpr (TIsCharEncodingCompatibleWith_V<UTF8CHAR, CharType>)
+	{
+		const int32 Offset = Builder.AddUninitialized(StringLength);
+		Ar.Serialize(Builder.GetData() + Offset, StringLength);
+	}
+	else
+	{
+		TArray<UTF8CHAR, TInlineAllocator<128>> Buffer;
+		Buffer.SetNumUninitialized(StringLength);
+		Ar.Serialize(Buffer.GetData(), StringLength);
+		const int32 ConvertedLength = FPlatformString::ConvertedLength<CharType>(Buffer.GetData(), StringLength);
+		const int32 Offset = Builder.AddUninitialized(ConvertedLength);
+		FPlatformString::Convert(Builder.GetData() + Offset, ConvertedLength, Buffer.GetData(), Buffer.Num());
+	}
+}
+
+/**
+ * Helper function to read a C# BinaryWriter-serialized string into a string.
+ * @tparam CharType Character type.
+ * @param Ar Binary archive to read the string from.
+ * @param OutString String to write to.
+ */
+template <typename CharType>
+void SerializeCsString(FArchive& Ar, TString<CharType>& OutString)
+{
+	// C# BinaryWriter prefixes strings with their length as a 7-bit encoded int
+	uint32 StringLength = Decode7bit(Ar);
+
+	TArray<CharType, typename TString<CharType>::AllocatorType>& Data = OutString.GetCharArray();
+	if constexpr (TIsCharEncodingCompatibleWith_V<UTF8CHAR, CharType>)
+	{
+		Data.SetNumUninitialized(StringLength + 1);
+		Ar.Serialize(Data.GetData(), StringLength);
+		Data[StringLength] = CHARTEXT(CharType, '\0');
+	}
+	else
+	{
+		TArray<UTF8CHAR, TInlineAllocator<128>> Buffer;
+		Buffer.SetNumUninitialized(StringLength);
+		Ar.Serialize(Buffer.GetData(), Buffer.Num());
+		const int32 ConvertedLength = FPlatformString::ConvertedLength<CharType>(Buffer.GetData(), Buffer.Num());
+		Data.SetNumUninitialized(ConvertedLength + 1);
+		FPlatformString::Convert(Data.GetData(), ConvertedLength, Buffer.GetData(), Buffer.Num());
+		Data[ConvertedLength] = CHARTEXT(CharType, '\0');
+	}
+}
+
+/**
+ * Helper function to read a C# BinaryWriter-serialized string into a string.
+ * @tparam CharType Character type.
+ * @param Ar Binary archive to read the string from.
+ * @return The string that was read.
+ */
+template <typename CharType>
+TString<CharType> SerializeCsString(FArchive& Ar)
+{
+	TString<CharType> String;
+	SerializeCsString<CharType>(Ar, String);
+	return String;
+}
+
+/**
+ * Helper function to serialize metadata from the parameter archive.
+ * @param Ar Binary archive to read the metadata from.
+ * @param OutMetadata The container to store the read metadata in.
+ */
+void SerializeMetadata(FArchive& Ar, TMap<FString, FString>& OutMetadata)
+{
+	int32 NumValues;
+	Ar << NumValues;
+
+	OutMetadata.Reserve(NumValues);
+	for (int32 Index = 0; Index < NumValues; Index++)
+	{
+		FString Key = SerializeCsString<TCHAR>(Ar);
+		FString Value = SerializeCsString<TCHAR>(Ar);
+		OutMetadata.Add(MoveTemp(Key), MoveTemp(Value));
+	}
+}
+
+/**
+ * Serializes binary CSV profiler data. Based on the C# code found in CsvStats.ReadBinFile.
+ * @param Ar Binary archive to read the capture data from.
+ * @param OutCapture Output container to store the read capture in.
+ * @param OutErrors Optional output array to hold any reported errors.
+ * @return True if serialization of the capture succeeded.
+ */
+bool SerializeBin(FArchive& Ar, FCsvProfilerCapture& OutCapture, TArray<FText>* OutErrors = nullptr)
+{
+	TRACE_CPUPROFILER_EVENT_SCOPE(Private::SerializeBin);
+
+	// Check magic
+	{
+		TUtf8StringBuilder<256> Builder;
+		SerializeCsString(Ar, Builder);
+		const FUtf8StringView CsvBinMagic(UTF8TEXTVIEW("CSVBIN"));
+		if (Builder != CsvBinMagic)
+		{
+			if (OutErrors)
+			{
+				OutErrors->Add(IncorrectFormatText);
+			}
+			return false;
+		}
+	}
+
+	ECsvBinVersion Version;
+	Ar << Version;
+	if (Version < ECsvBinVersion::InitialRelease)
+	{
+		if (OutErrors)
+		{
+			OutErrors->Add(UnsupportedVersionText);
+		}
+		return false;
+	}
+
+	// Read flags
+	ECsvBinFlags Flags = ECsvBinFlags::None;
+	bool bCompressed = false;
+	if (Version >= ECsvBinVersion::CompressionSupportAndFlags)
+	{
+		Ar << Flags;
+
+		ECsvBinCompressionLevel CompressionLevel;
+		Ar << CompressionLevel;
+		bCompressed = CompressionLevel != ECsvBinCompressionLevel::None;
+
+		if (bCompressed)
+		{
+			CsvBinCompressionType CompressionType;
+			Ar << CompressionType;
+			if (CompressionType != CsvBinCompressionType::MsDeflate)
+			{
+				if (OutErrors)
+				{
+					OutErrors->Add(UnsupportedCompressionTypeText);
+				}
+				return false;
+			}
+		}
+		else
+		{
+			if (OutErrors)
+			{
+				OutErrors->Add(UncompressedFormatSupportText);
+			}
+			return false;
+		}
+	}
+	else
+	{
+		bool bHasMetadata = false;
+		Ar << bHasMetadata;
+		if (bHasMetadata)
+		{
+			EnumAddFlags(Flags, ECsvBinFlags::HasMetadata);
+		}
+	}
+
+	TMap<FString, FString> Metadata;
+	if (EnumHasAnyFlags(Flags, ECsvBinFlags::HasMetadata))
+	{
+		SerializeMetadata(Ar, Metadata);
+	}
+
+	// Read counts
+	int32 EventCount, ValueCount, SampleCount;
+	Ar << EventCount << ValueCount << SampleCount;
+
+	// Read sample names
+	TMap<FString, FCsvProfilerSample> Samples;
+	Samples.Reserve(SampleCount);
+	for (int32 Index = 0; Index < SampleCount; Index++)
+	{
+		Samples.Add(SerializeCsString<TCHAR>(Ar));
+	}
+
+	// Read the sample data
+	for (int32 Index = 0; Index < SampleCount; Index++)
+	{
+		FString SampleName = SerializeCsString<TCHAR>(Ar);
+
+		FCsvProfilerSample* FoundSample = Samples.Find(SampleName);
+		if (!FoundSample)
+		{
+			if (OutErrors)
+			{
+				OutErrors->Add(
+					FText::FormatNamed(SampleDataNotFoundText, FString(SampleNameKey), FText::FromString(SampleName)));
+			}
+			return false;
+		}
+
+		FCsvProfilerSample& Sample = *FoundSample;
+
+		Sample.Name = MoveTemp(SampleName);
+		Ar << Sample.Average;
+		Ar << Sample.Total;
+
+		int32 StatSizeBytes;
+		Ar << StatSizeBytes;
+
+		if (bCompressed)
+		{
+			int32 CompressedBufferLength;
+			Ar << CompressedBufferLength;
+
+			TArray<uint8> CompressedBuffer;
+			CompressedBuffer.SetNumUninitialized(CompressedBufferLength);
+			Ar.Serialize(CompressedBuffer.GetData(), CompressedBufferLength);
+
+			int32 UncompressedBufferLength = sizeof(float) * ValueCount;
+			TArray<float> UncompressedValuesBuffer;
+			UncompressedValuesBuffer.SetNumUninitialized(ValueCount);
+
+			// FCompression::UncompressMemory uses the CompressionData parameter as the window size zlib uses.
+			// Per zlib's documentation, a window size value between -8 and -15 will do a raw inflate, which
+			// is what we want as the compressed data doesn't have any headers.
+			constexpr int32 WindowSizeRawDeflate = -15;
+			if (!FCompression::UncompressMemory(
+					/*          FormatName = */ NAME_Zlib,
+					/*  UncompressedBuffer = */ UncompressedValuesBuffer.GetData(),
+					/*    UncompressedSize = */ UncompressedBufferLength,
+					/*    CompressedBuffer = */ CompressedBuffer.GetData(),
+					/*      CompressedSize = */ CompressedBufferLength,
+					/*               Flags = */ COMPRESS_NoFlags,
+					/*     CompressionData = */ WindowSizeRawDeflate))
+			{
+				return false;
+			}
+
+			Sample.Values = MoveTemp(UncompressedValuesBuffer);
+		}
+		else
+		{
+			unimplemented();
+		}
+	}
+
+	// Read the event data
+	TArray<FCsvProfilerEvent> Events;
+	Events.Reserve(SampleCount);
+	for (int i = 0; i < EventCount; i++)
+	{
+		int32 Frame;
+		Ar << Frame;
+		FString Name = SerializeCsString<TCHAR>(Ar);
+		Events.Add({.Name = MoveTemp(Name), .Frame = Frame});
+	}
+
+	// Success - move into the output capture container.
+	OutCapture.Samples = MoveTemp(Samples);
+	OutCapture.Events = MoveTemp(Events);
+	OutCapture.Metadata = MoveTemp(Metadata);
+	return true;
+}
+
+}  // namespace CsvUtil::Private
+
+bool ReadFromCsv(FCsvProfilerCapture& OutCapture, const TCHAR* FilePath, TArray<FText>* OutErrors /*= nullptr*/)
+{
+	TRACE_CPUPROFILER_EVENT_SCOPE(ReadFromCsv);
+	FScopedSlowTask SlowTask(1, LOCTEXT("ReadFromCsv", "Reading CSV data"));
+	SlowTask.MakeDialog();
+	SlowTask.EnterProgressFrame();
+	OutCapture = FCsvProfilerCapture{};
+	if (CsvUtils::Private::FCsvProfilerFileReader Reader(OutErrors);
+		FFileHelper::LoadFileToStringWithLineVisitor(FilePath, Reader))
+	{
+		return Reader.Finalize(OutCapture);
+	}
+	return false;
+}
+
+bool ReadFromCsvBin(FCsvProfilerCapture& OutCapture, const TCHAR* FilePath, TArray<FText>* OutErrors /*= nullptr*/)
+{
+	TRACE_CPUPROFILER_EVENT_SCOPE(ReadFromCsvBin);
+	FScopedSlowTask SlowTask(1, LOCTEXT("ReadFromCsvBin", "Reading CSV data"));
+	SlowTask.MakeDialog();
+	SlowTask.EnterProgressFrame();
+	OutCapture = FCsvProfilerCapture{};
+	if (TUniquePtr<FArchive> Archive(IFileManager::Get().CreateFileReader(FilePath)); Archive.IsValid())
+	{
+		return CsvUtils::Private::SerializeBin(*Archive, OutCapture, OutErrors);
+	}
+	return false;
+}
+
+}
\ No newline at end of file
diff --git a/Engine/Source/Programs/CSVTools/CSVUtils/Public/CSVProfilerUtils.h b/Engine/Source/Programs/CSVTools/CSVUtils/Public/CSVProfilerUtils.h
new file mode 100644
--- /dev/null
+++ b/Engine/Source/Programs/CSVTools/CSVUtils/Public/CSVProfilerUtils.h
@@ -0,0 +1,50 @@
+// Copyright Epic Games, Inc. All Rights Reserved.
+
+#pragma once
+
+#include "Containers/Map.h"
+#include "Containers/UnrealString.h"
+
+namespace CsvUtils
+{
+
+	struct FCsvProfilerSample
+	{
+		FString Name;
+		TArray<float> Values;
+		float Average = 0.0f;
+		double Total = 0.0;
+	};
+
+	struct FCsvProfilerEvent
+	{
+		FString Name;
+		int32 Frame = INDEX_NONE;
+	};
+
+	struct FCsvProfilerCapture
+	{
+		TMap<FString, FCsvProfilerSample> Samples;
+		TArray<FCsvProfilerEvent> Events;
+		TMap<FString, FString> Metadata;
+	};
+
+	/**
+	 * Reads a csv profiler capture from the parameter file path into the parameter container.
+	 * @param OutCapture The destination capture container.
+	 * @param FilePath The source file path.
+	 * @param OutErrors Optional output array to hold any reported errors.
+	 * @return True if the capture was successfully read.
+	 */
+	CSVUTILS_API bool ReadFromCsv(FCsvProfilerCapture& OutCapture, const TCHAR* FilePath, TArray<FText>* OutErrors = nullptr);
+
+	/**
+	 * Reads a binary csv profiler capture from the parameter file path into the parameter container.
+	 * @param OutCapture The destination capture container.
+	 * @param FilePath The source file path.
+	 * @param OutErrors Optional output array to hold any reported errors.
+	 * @return True if the capture was successfully read.
+	 */
+	CSVUTILS_API bool ReadFromCsvBin(FCsvProfilerCapture& OutCapture, const TCHAR* FilePath, TArray<FText>* OutErrors = nullptr);
+
+}
\ No newline at end of file
