Added missing #undef LOCTEXT_NAMESPACE directive

diff --git a/Engine/Source/Programs/CSVTools/CSVUtils/Private/CSVProfilerUtils.cpp b/Engine/Source/Programs/CSVTools/CSVUtils/Private/CSVProfilerUtils.cpp
--- a/Engine/Source/Programs/CSVTools/CSVUtils/Private/CSVProfilerUtils.cpp
+++ b/Engine/Source/Programs/CSVTools/CSVUtils/Private/CSVProfilerUtils.cpp
@@ -479,4 +479,6 @@ bool ReadFromCsvBin(FCsvProfilerCapture& OutCapture, const TCHAR* FilePath, TArr
 	return false;
 }
 
-}
\ No newline at end of file
+}
+
+#undef LOCTEXT_NAMESPACE
\ No newline at end of file
