CSVToTelemetry Quality Improvements
Changed the error log message from Display to Error.
Added -help to display the expected usage and optional switches.
Added upload time summary.

diff --git a/Engine/Source/Programs/CSVTools/CSVToTelemetry/Private/CSVToTelemetry.cpp b/Engine/Source/Programs/CSVTools/CSVToTelemetry/Private/CSVToTelemetry.cpp
--- a/Engine/Source/Programs/CSVTools/CSVToTelemetry/Private/CSVToTelemetry.cpp
+++ b/Engine/Source/Programs/CSVTools/CSVToTelemetry/Private/CSVToTelemetry.cpp
@@ -208,6 +208,12 @@ static int GenerateTelemetryFromCSVProfileFile(const FString& FilePath)
 	return Result;
 }
 
+static int ShowHelp()
+{
+	UE_LOG(LogCSVToTelemetry, Display, TEXT("\n\nCSVToTelemetry Help\n\nUsage:\n\tCSVToTelemetry.exe -csv=[filename] -event=[eventname] ( -schema=[value] -column=[name1|name2|....] )\n\tCSVToTelemetry.exe -csvprofile=[filename] -event=[eventname]\n\tCSVToTelemetry.exe -help\n\nRequired:\n\t-csv=[filename]\t\t\tGeneric text based csv input file.\n\t-csvprofile=[filename]\t\tCSVProfiler csv input file. Denote binary with .csv.bin otherwise assumed text.\n\t-event=[name]\t\t\tName of telemetry event to send each row to.\nOptional:\n\t-schema=[value]\t\t\tEvent schema value.\n\t-columns=[name1|name2|....]\tColumn include filter." ) );
+	return 0;
+}
+
 INT32_MAIN_INT32_ARGC_TCHAR_ARGV()
 {
 	FTaskTagScope Scope(ETaskTag::EGameThread);
@@ -224,24 +230,33 @@ INT32_MAIN_INT32_ARGC_TCHAR_ARGV()
 
 	FString FilePath;
 	
-	if (FParse::Value(FCommandLine::Get(), TEXT("csvprofile="), FilePath) == true)
+	if (FParse::Value(FCommandLine::Get(), TEXT("help"), FilePath) == true)
 	{
-		Result = GenerateTelemetryFromCSVProfileFile(FilePath);
+		// Show help message
+		Result = ShowHelp();
 	}
-	else if (FParse::Value(FCommandLine::Get(), TEXT("csv="), FilePath)==true)
-	{
-		Result = GenerateTelemetryFromCSVFile(FilePath);
+	else
+	{ 
+		if (FParse::Value(FCommandLine::Get(), TEXT("csvprofile="), FilePath) == true)
+		{
+			Result = GenerateTelemetryFromCSVProfileFile(FilePath);
+		}
+		else if (FParse::Value(FCommandLine::Get(), TEXT("csv="), FilePath) == true)
+		{
+			Result = GenerateTelemetryFromCSVFile(FilePath);
+		}
+	
+		if (Result == 0)
+		{
+			// Upload completed successfully
+			UE_LOG(LogCSVToTelemetry, Display, TEXT("CSVToTelemetry upload succeeded in %0.2f seconds"), (FDateTime::UtcNow() - StartTime).GetTotalSeconds());
+		}
 	}
 
-	if (Result!=0)
+	if (Result != 0)
 	{
 		// Always show the usage example if we have not been successful
-		UE_LOG(LogCSVToTelemetry, Display, TEXT("Usage CSVToTelemetry.exe -csv=[filename] -event=[eventname] ( -schema=[value] -columns=[name1|name2|....] )"));
-	}
-	else
-	{
-		// Upload completed successfully
-		UE_LOG(LogCSVToTelemetry, Display, TEXT("CSVToTelemetry upload succeeded in %0.2f seconds"), (FDateTime::UtcNow() - StartTime).GetTotalSeconds());
+		UE_LOG(LogCSVToTelemetry, Error, TEXT("\nUsage:\n\tCSVToTelemetry.exe -csv=[filename] -event=[eventname] ( -schema=[value] -columns=[name1|name2|....] )\n\tCSVToTelemetry.exe -csvprofile=[filename] -event=[eventname]\n\tCSVToTelemetry.exe -help"));
 	}
 	
 	if (FParse::Param(FCommandLine::Get(), TEXT("fastexit")))
