Update the performance tool to have correct case sensitivity for the Report Types and Graphs and also the output directory

Test succeeds and correctly outputs the summary graphs now

diff --git a/Engine/Source/Programs/CSVTools/PerfReportTool/PerfReportTool.cs b/Engine/Source/Programs/CSVTools/PerfReportTool/PerfReportTool.cs
--- a/Engine/Source/Programs/CSVTools/PerfReportTool/PerfReportTool.cs
+++ b/Engine/Source/Programs/CSVTools/PerfReportTool/PerfReportTool.cs
@@ -23,7 +23,7 @@ namespace PerfReportTool
     class Version
     {
 		// Format: Major.Minor.Bugfix
-        private static string VersionString = "4.244.3";
+        private static string VersionString = "4.244.4";
 
         public static string Get() { return VersionString; }
     };
diff --git a/Engine/Source/Programs/CSVTools/PerfReportTool/ReportXml.cs b/Engine/Source/Programs/CSVTools/PerfReportTool/ReportXml.cs
--- a/Engine/Source/Programs/CSVTools/PerfReportTool/ReportXml.cs
+++ b/Engine/Source/Programs/CSVTools/PerfReportTool/ReportXml.cs
@@ -177,7 +177,7 @@ namespace PerfReportTool
 			baseXmlDirectory = baseDirectory;
 
 			// Read the report type XML
-			reportTypeXmlFilename = Path.Combine(baseDirectory, "reportTypes.xml");
+			reportTypeXmlFilename = Path.Combine(baseDirectory, "ReportTypes.xml");
 			if (reportXMLFilenameIn.Length > 0)
 			{
 				reportTypeXmlFilename = IsAbsolutePath(reportXMLFilenameIn) ? reportXMLFilenameIn : Path.Combine(baseDirectory, reportXMLFilenameIn);
@@ -250,7 +250,7 @@ namespace PerfReportTool
 				}
 				else
 				{
-					graphsXMLFilename = Path.Combine(baseDirectory, "reportGraphs.xml");
+					graphsXMLFilename = Path.Combine(baseDirectory, "ReportGraphs.xml");
 				}
 
 			}
@@ -999,4 +999,4 @@ namespace PerfReportTool
 			}
 		}
 	};
-}
\ No newline at end of file
+}
