Fix another issue with the case sensitivity
- We were lowering the directory path unncessarily
- This should have been in the previous check-in but somehow didn't make it after shelving (for some reason)

diff --git a/Engine/Source/Programs/CSVTools/PerfReportTool/PerfReportTool.cs b/Engine/Source/Programs/CSVTools/PerfReportTool/PerfReportTool.cs
--- a/Engine/Source/Programs/CSVTools/PerfReportTool/PerfReportTool.cs
+++ b/Engine/Source/Programs/CSVTools/PerfReportTool/PerfReportTool.cs
@@ -23,7 +23,7 @@ namespace PerfReportTool
     class Version
     {
 		// Format: Major.Minor.Bugfix
-        private static string VersionString = "4.244.4";
+        private static string VersionString = "4.244.5";
 
         public static string Get() { return VersionString; }
     };
@@ -508,7 +508,7 @@ namespace PerfReportTool
 			}
 
 			// Create the output directory if requested
-			string outputDir = GetArg("o", false).ToLower();
+			string outputDir = GetArg("o", false);
 			if (!string.IsNullOrEmpty(outputDir))
 			{
 				if (!Directory.Exists(outputDir))
