[CsvUtils] Update C++ CSV text serialization logic to more closely match the C# version, improving correctness and performance at the cost of additional memory usage.

#tests Opened CSV profiler data from 2 sources: .csv output directly from a game, and .csv converted from .csv.bin via CsvConvert.

diff --git a/Engine/Source/Programs/CSVTools/CSVUtils/Private/CSVProfilerUtils.cpp b/Engine/Source/Programs/CSVTools/CSVUtils/Private/CSVProfilerUtils.cpp
--- a/Engine/Source/Programs/CSVTools/CSVUtils/Private/CSVProfilerUtils.cpp
+++ b/Engine/Source/Programs/CSVTools/CSVUtils/Private/CSVProfilerUtils.cpp
@@ -2,6 +2,7 @@
 
 #include "CSVProfilerUtils.h"
 
+#include "Async/ParallelFor.h"
 #include "HAL/FileManager.h"
 #include "Misc/Compression.h"
 #include "Misc/FileHelper.h"
@@ -29,16 +30,8 @@ static const FText UncompressedFormatSupportText(
 	LOCTEXT("UncompressedFormatSupport", "Uncompressed format loading is not yet supported."));
 static const FText SampleDataNotFoundText(
 	LOCTEXT("SampleDataNotFound", "Unable to find sample with name {SampleName} to serialize."));
-static const FText SourceDataNotFoundText(
-	LOCTEXT("SourceDataNotFound", "Unable to find data with name {SourceName} to populate sample."));
-static const FText FailedToConvertValueText(
-	LOCTEXT("FailedToConvertValue", "Failed to convert value {Value} from {SourceUnit} to {DestUnit}."));
 
 static constexpr FStringView SampleNameKey(TEXTVIEW("SampleName"));
-static constexpr FStringView SourceNameKey(TEXTVIEW("SourceName"));
-static constexpr FStringView ValueKey(TEXTVIEW("Value"));
-static constexpr FStringView SourceUnitKey(TEXTVIEW("SourceUnit"));
-static constexpr FStringView DestUnitKey(TEXTVIEW("DestUnit"));
 
 // Should be kept up-to-date with CSVStats.CsvBinVersion
 enum class ECsvBinVersion : int32
@@ -73,78 +66,192 @@ enum class ECsvBinCompressionLevel : uint8
 	Max
 };
 
-struct FCsvProfilerFileReader
+static bool LineIsMetadata(FStringView Line)
 {
-	explicit FCsvProfilerFileReader(TArray<FText>* InOutErrors) : OutErrors(InOutErrors)
+	if (Line.TrimStartAndEnd().StartsWith(TEXT('[')))
 	{
+		return true;
+	}
+	return false;
+}
+
+/**
+* Helper function to serialize metadata from the parameter line.
+* @param Line Text line to read the metadata from.
+* @param OutMetadata The container to store the read metadata in.
+*/
+void SerializeMetadataText(const FString& Line, TMap<FString, FString>& OutMetadata)
+{
+	// Initialize state tracking variables
+	bool bIsKey = false;
+	FString CurrentKey;
+
+	// Split the metadata line into segments at commas
+	const TCHAR* LinePtr = *Line;
+	FString Token;
+	while (FParse::Token(LinePtr, Token, false, ','))
+	{		
+		// Check if this is a key (enclosed in square brackets)
+		if (!bIsKey && Token.StartsWith(TEXT("[")) && Token.EndsWith(TEXT("]")))
+		{
+			// Extract key without brackets and convert to lowercase
+			CurrentKey = MoveTemp(Token);
+			CurrentKey.MidInline(1, CurrentKey.Len() - 2);
+			CurrentKey.ToLowerInline();
+			bIsKey = true;
+			continue;
+		}
+
+		// Handle the value
+		if (bIsKey)
+		{
+			// Add or append to existing value
+			if (FString* ExistingValue = OutMetadata.Find(CurrentKey))
+			{
+				*ExistingValue += TEXT(",") + Token;
+			}
+			else
+			{
+				OutMetadata.Add(MoveTemp(CurrentKey), MoveTemp(Token));
+			}
+
+			bIsKey = false;
+		}
+	}
+}
+
+/**
+ * Serializes binary CSV profiler data. Based on the C# code found in CsvStats.ReadCSVFromLines.
+ * @param Lines Text lines to serialize from.
+ * @param OutCapture Output container to store the read capture in.
+ * @param OutErrors Optional output array to hold any reported errors.
+ * @return True if serialization of the capture succeeded.
+ */
+bool SerializeText(TConstArrayView<FString> Lines, FCsvProfilerCapture& OutCapture, TArray<FText>* OutErrors = nullptr)
+{
+	TRACE_CPUPROFILER_EVENT_SCOPE(Private::SerializeText);
+
+	if (Lines.IsEmpty())
+	{
+		return false;
+	}
+
+	// Use the first line as the header view, unless the metadata tells us to use the row at the end instead.
+	const FString* HeaderRow = &Lines[0];
+
+	// Remove the header row from the view.
+	Lines.RightChopInline(1);
+
+	TMap<FString, FString> Metadata;
+	if (LineIsMetadata(Lines.Last()))
+	{
+		TRACE_CPUPROFILER_EVENT_SCOPE(Private::SerializeText::Metadata);
+
+		// Serialize the metadata.
+		SerializeMetadataText(Lines.Last(), Metadata);
+
+		// Remove the metadata line.
+		Lines.LeftChopInline(1);
+
+		// New CSVs from the csv profiler have a header row at the end of the file,
+		// since the profiler writes out the file incrementally.
+		const FString* HasHeaderRowAtEndValue = Metadata.Find("hasheaderrowatend");
+		if (HasHeaderRowAtEndValue && *HasHeaderRowAtEndValue == "1")
+		{
+			// Swap the header row for the one at the end of the file.
+			HeaderRow = &Lines[Lines.Num() - 1];
+
+			// Remove the end header row.
+			Lines.LeftChopInline(1);
+		}
 	}
 
-	TArray<FText>* OutErrors = nullptr;
+	// We should be left with only sample lines.
+	const int32 NumSamples = Lines.Num();
 
-	int32 FrameNumber = INDEX_NONE;
 	TArray<FString> SampleNames;
 	TArray<TArray<float>> SampleData;
-	TArray<FCsvProfilerEvent> Events;
 
-	void operator()(FStringView Line)
-	{
-		TRACE_CPUPROFILER_EVENT_SCOPE(Private::FCsvFileReader::operator());
+	int32 EventsHeadingIndex = INDEX_NONE;
 
-		// Metadata line
-		if (UNLIKELY(Line.StartsWith(TEXT('['))))
+	// Headers
+	{
+		TRACE_CPUPROFILER_EVENT_SCOPE(Private::SerializeText::Headers);
+		TArray<FString> Headings;
+		HeaderRow->ParseIntoArray(Headings, TEXT(","), false);
+		SampleNames.SetNum(Headings.Num());
+		SampleData.SetNum(Headings.Num());
+		for (TArray<FString>::TIterator Iter(Headings); Iter; ++Iter)
 		{
-			// unimplemented
+			const int32 ColumnIndex = Iter.GetIndex();
+			if (FStringView(*Iter).Equals(TEXTVIEW("events"), ESearchCase::IgnoreCase))
+			{
+				EventsHeadingIndex = ColumnIndex;
+			}
+			else
+			{
+				SampleNames[ColumnIndex] = MoveTemp(*Iter);
+				SampleData[ColumnIndex].SetNum(NumSamples);
+			}
 		}
-		else
+	}
+
+	TArray<FString> EventStrings;
+	EventStrings.SetNum(NumSamples);
+
+	// Samples
+	{
+		TRACE_CPUPROFILER_EVENT_SCOPE(Private::SerializeText::Samples);
+		ParallelFor(NumSamples, [&](int32 Index)
 		{
-			// Process a CSV line. Read each token in between commas.
-			// Note that this doesn't properly handle commas in quotes, but we don't
-			// expect any in the data we're processing.
-			bool bHeadingsRow = false;
-			int32 ColumnIndex = INDEX_NONE;
-			ParseTokens(
-				Line, TEXT(','),
-				[&](FStringView Token)
+			TRACE_CPUPROFILER_EVENT_SCOPE(Private::SerializeText::Samples::Work);
+			// Split the metadata line into segments at commas
+			const TCHAR* LinePtr = *Lines[Index];
+			FString Token;
+			int32 ColumnIndex = 0;
+			while (FParse::Token(LinePtr, Token, false, TEXT(',')))
+			{
+				if (LIKELY(ColumnIndex != EventsHeadingIndex))
 				{
-					++ColumnIndex;
-					// The first column is known - it contains events. So we can check if this is a headings row
-					// by checking if this column contains this label.
-					if (ColumnIndex == 0)
-					{
-						bHeadingsRow = Token == TEXTVIEW("EVENTS");
-						if (!bHeadingsRow)
-						{
-							// Bump the frame number for each non-headings row we come across.
-							++FrameNumber;
-							//  Split events on the delimiter.
-							ParseTokens(
-								Token, TEXT(';'),
-								[&](FStringView EventToken)
-								{ Events.Add({.Name = FString(EventToken), .Frame = FrameNumber}); },
-								UE::String::EParseTokensOptions::IgnoreCase | UE::String::EParseTokensOptions::SkipEmpty);
-						}
-						return;
-					}
-					const int32 SampleIndex = ColumnIndex - 1;
-					if (UNLIKELY(bHeadingsRow))
-					{
-						SampleNames.SetNum(FMath::Max(SampleNames.Num(), SampleIndex + 1));
-						SampleNames[SampleIndex] = Token;
-					}
-					else
-					{
-						SampleData.SetNum(FMath::Max(SampleData.Num(), ColumnIndex));
-						SampleData[SampleIndex].SetNumZeroed(FrameNumber + 1);
-						LexFromString(SampleData[SampleIndex][FrameNumber], Token);
-					}
-				},
-				UE::String::EParseTokensOptions::IgnoreCase);
+					LexFromString(SampleData[ColumnIndex][Index], Token);
+				}
+				else
+				{
+					EventStrings[Index] = MoveTemp(Token);
+				}
+
+				++ColumnIndex;
+			}
+		});
+	}
+
+	// Events
+	TArray<FCsvProfilerEvent> Events;
+	if (EventsHeadingIndex != INDEX_NONE)
+	{
+		TRACE_CPUPROFILER_EVENT_SCOPE(Private::SerializeText::Events);
+		for (TArray<FString>::TConstIterator Iter(EventStrings); Iter; ++Iter)
+		{
+			if (Iter->IsEmpty())
+			{
+				continue;
+			}
+
+			ParseTokens(
+				*Iter, TEXT(';'),
+				[&](FStringView EventToken)
+				{ Events.Add({ .Name = FString(EventToken), .Frame = Iter.GetIndex() }); },
+				UE::String::EParseTokensOptions::IgnoreCase | UE::String::EParseTokensOptions::SkipEmpty);
 		}
+
+		// Cleanup unused stat data for events
+		SampleNames.RemoveAt(EventsHeadingIndex, EAllowShrinking::No);
+		SampleData.RemoveAt(EventsHeadingIndex, EAllowShrinking::No);
 	}
 
-	bool Finalize(FCsvProfilerCapture& OutCapture)
+	// Finalize
 	{
-		TRACE_CPUPROFILER_EVENT_SCOPE(Private::FCsvFileReader::Finalize);
+		TRACE_CPUPROFILER_EVENT_SCOPE(Private::SerializeText::Finalize);
 		OutCapture.Samples.Reserve(SampleNames.Num());
 		for (int32 SeriesIndex = 0; SeriesIndex < SampleNames.Num(); ++SeriesIndex)
 		{
@@ -153,9 +260,10 @@ struct FCsvProfilerFileReader
 			Sample.Values = MoveTemp(SampleData[SeriesIndex]);
 		}
 		OutCapture.Events = MoveTemp(Events);
-		return true;
+		OutCapture.Metadata = MoveTemp(Metadata);
 	}
-};
+	return true;
+}
 
 /**
  * Helper function to decode a 7-bit encoded integer from the parameter archive.
@@ -256,7 +364,7 @@ TString<CharType> SerializeCsString(FArchive& Ar)
  * @param Ar Binary archive to read the metadata from.
  * @param OutMetadata The container to store the read metadata in.
  */
-void SerializeMetadata(FArchive& Ar, TMap<FString, FString>& OutMetadata)
+void SerializeMetadataBin(FArchive& Ar, TMap<FString, FString>& OutMetadata)
 {
 	int32 NumValues;
 	Ar << NumValues;
@@ -353,7 +461,7 @@ bool SerializeBin(FArchive& Ar, FCsvProfilerCapture& OutCapture, TArray<FText>*
 	TMap<FString, FString> Metadata;
 	if (EnumHasAnyFlags(Flags, ECsvBinFlags::HasMetadata))
 	{
-		SerializeMetadata(Ar, Metadata);
+		SerializeMetadataBin(Ar, Metadata);
 	}
 
 	// Read counts
@@ -457,10 +565,10 @@ bool ReadFromCsv(FCsvProfilerCapture& OutCapture, const TCHAR* FilePath, TArray<
 	SlowTask.MakeDialog();
 	SlowTask.EnterProgressFrame();
 	OutCapture = FCsvProfilerCapture{};
-	if (CsvUtils::Private::FCsvProfilerFileReader Reader(OutErrors);
-		FFileHelper::LoadFileToStringWithLineVisitor(FilePath, Reader))
+	TArray<FString> Lines;
+	if (FFileHelper::LoadFileToStringArray(Lines, FilePath))
 	{
-		return Reader.Finalize(OutCapture);
+		return CsvUtils::Private::SerializeText(Lines, OutCapture, OutErrors);
 	}
 	return false;
 }
