CSVTools: Added support for adding custom csvColor overrides when using Csv2Svg

diff --git a/Engine/Source/Programs/CSVTools/CSVToSVG/CsvToSvgProgram.cs b/Engine/Source/Programs/CSVTools/CSVToSVG/CsvToSvgProgram.cs
--- a/Engine/Source/Programs/CSVTools/CSVToSVG/CsvToSvgProgram.cs
+++ b/Engine/Source/Programs/CSVTools/CSVToSVG/CsvToSvgProgram.cs
@@ -105,10 +105,11 @@ namespace CSVTools
             "       -percentile90 \n" +
             "       -percentile99 \n" +
 			"       -uniqueId <string> : unique ID for JS (needed if this is getting embedded in HTML alongside other graphs)\n" +
-			"       -nocommandlineEmbed : don't embed the commandline in the SVG" +
-			"       -lineDecimalPlaces <N> (only effective with with -fixedPointGraphs 0. Default=3)" +
-			"       -frameOffset <N> : offset used for frame display name (default 0)" +
-            "";
+			"       -nocommandlineEmbed : don't embed the commandline in the SVG\n" +
+			"       -lineDecimalPlaces <N> (only effective with with -fixedPointGraphs 0. Default=3)\n" +
+			"       -frameOffset <N> : offset used for frame display name (default 0)\n" +
+			"       -csvColors <list> : list of hex colors. Applied in order to corresponding csvs when used with a list of csv files or csvDir\n" +
+			"";
 
 		void Run(string[] args)
 		{
@@ -332,6 +333,11 @@ namespace CSVTools
 			graphParams.ignoreStats = GetListArg("ignoreStats");
 			graphParams.themeName = GetArg("theme").ToLower();
 			graphParams.colorOffset = GetIntArg("colorOffset", GetIntArg("colourOffset", 0) );
+			graphParams.csvColors = GetListArg("csvColors");
+			if(graphParams.csvColors.Count == 0)
+			{
+				graphParams.csvColors = GetListArg("csvColours");
+			}
 
 			// Events
 			graphParams.showEventNames = GetListArg("showEvents");
diff --git a/Engine/Source/Programs/CSVTools/CsvStats/CsvStats.cs b/Engine/Source/Programs/CSVTools/CsvStats/CsvStats.cs
--- a/Engine/Source/Programs/CSVTools/CsvStats/CsvStats.cs
+++ b/Engine/Source/Programs/CSVTools/CsvStats/CsvStats.cs
@@ -2519,7 +2519,17 @@ namespace CSVStats
         }
         public Colour(Colour colourIn) { r = colourIn.r; g = colourIn.g; b = colourIn.b; alpha = colourIn.alpha; }
 
-        public string SVGString()
+		public Colour(string str)
+		{
+			string hexStr = str.TrimStart('#');
+			int hex = Convert.ToInt32(hexStr, 16);
+			r = (byte)((hex >> 16) & 0xff);
+			g = (byte)((hex >> 8) & 0xff);
+			b = (byte)((hex >> 0) & 0xff);
+			alpha = 1.0f;
+		}
+
+		public string SVGString()
         {
             return "'" + SVGStringNoQuotes() + "'";
         }
diff --git a/Engine/Source/Programs/CSVTools/CsvToSvgLib/CsvToSvgLib.cs b/Engine/Source/Programs/CSVTools/CsvToSvgLib/CsvToSvgLib.cs
--- a/Engine/Source/Programs/CSVTools/CsvToSvgLib/CsvToSvgLib.cs
+++ b/Engine/Source/Programs/CSVTools/CsvToSvgLib/CsvToSvgLib.cs
@@ -102,6 +102,7 @@ namespace CSVTools
 		public bool forceLegendSort = false;
 		public float legendAverageThreshold = -float.MaxValue;
 		public List<string> customLegendNames = new List<string>();
+		public List<string> csvColors = new List<string>();
 
 		// Process options (doesn't affect output)
 		public bool bSmoothMultithreaded = true;
@@ -498,6 +499,15 @@ namespace CSVTools
 				CsvInfo csvInfo = csvList[i];
 				FrameRange csvFrameRange;
 				CsvStats newCsvStats = ProcessCsvStats(csvInfo.stats, graphParams, out csvFrameRange);
+				Colour csvColorOverride = null;
+
+				// Do we have specific csv color overrides?
+				if(graphParams.csvColors.Count > 0)
+				{
+					int csvColorIndex = (i % graphParams.csvColors.Count);
+					// assign a color override for this CSV
+					csvColorOverride = new Colour(graphParams.csvColors[csvColorIndex]);
+				}
 
 				if ( csvFrameRange.isLimited() )
 				{
@@ -548,7 +558,7 @@ namespace CSVTools
 				SetLegend(newCsvStats, csvInfo.filename, graphParams, csvList.Count > 1, ref currentCustomLabelIndex);
 				if (!graphParams.stacked)
 				{
-					currentColorOffset = AssignColours(newCsvStats, theme, false, currentColorOffset);
+					currentColorOffset = AssignColours(newCsvStats, theme, false, currentColorOffset, csvColorOverride);
 				}
 				csvStatsList.Add(newCsvStats);
 			}
@@ -609,7 +619,7 @@ namespace CSVTools
 				if (graphParams.stacked)
 				{
 					csvStatsList[i] = StackStats(csvStatsList[i], range, graphParams.stackedUnsorted, stackTotalStat, stackTotalStatIsAutomatic);
-					AssignColours(csvStatsList[i], theme, true, graphParams.colorOffset);
+					AssignColours(csvStatsList[i], theme, true, graphParams.colorOffset, null);
 
 					// Copy the colours to the unstacked stats
 					foreach (StatSamples samples in unstackedCsvStats[i].Stats.Values)
@@ -1036,7 +1046,7 @@ namespace CSVTools
 			}
 		}
 
-		int AssignColours(CsvStats stats, Theme theme, bool reverseOrder, int inColorOffset)
+		int AssignColours(CsvStats stats, Theme theme, bool reverseOrder, int inColorOffset, Colour csvColorOverride)
 		{
 			int colorOffset = inColorOffset;
 			if (reverseOrder)
@@ -1046,21 +1056,29 @@ namespace CSVTools
 
 			foreach (StatSamples stat in stats.Stats.Values)
 			{
-				if (stat.colour.alpha == 0.0f)
+				if(csvColorOverride == null)
 				{
-					stat.colour = theme.GraphColours[colorOffset % theme.GraphColours.Length];
-				}
-				if (reverseOrder)
-				{
-					colorOffset--;
+					if (stat.colour.alpha == 0.0f)
+					{
+						stat.colour = theme.GraphColours[colorOffset % theme.GraphColours.Length];
+					}
+					if (reverseOrder)
+					{
+						colorOffset--;
+					}
+					else
+					{
+						colorOffset++;
+					}
 				}
 				else
 				{
-					colorOffset++;
+					stat.colour = csvColorOverride;
 				}
 			}
 			return colorOffset;
 		}
+
 		CsvStats StackStats(CsvStats stats, Range range, bool stackedUnsorted, string stackTotalStat, bool bStackTotalStatIsAutomatic)
 		{
 			// Find our stack total stat
