PerfReportTool - Map overlay: fix event color changing edge case when the first event is on frame 0

diff --git a/Engine/Source/Programs/CSVTools/PerfReportTool/Summaries/MapOverlaySummary.cs b/Engine/Source/Programs/CSVTools/PerfReportTool/Summaries/MapOverlaySummary.cs
--- a/Engine/Source/Programs/CSVTools/PerfReportTool/Summaries/MapOverlaySummary.cs
+++ b/Engine/Source/Programs/CSVTools/PerfReportTool/Summaries/MapOverlaySummary.cs
@@ -238,7 +238,7 @@ namespace PerfSummaries
 					float adjustedLineSplitDistanceThreshold = lineSplitDistanceThreshold * framesPerLineSegment;
 					float oldx = 0;
 					float oldy = 0;
-					int lastFrameIndex = 0;
+					int lastFrameIndex = -1;
 					for (int i = startFrame; i < xStat.samples.Count; i += framesPerLineSegment)
 					{
 						float x = xStat.samples[i];
