PerfReportTool 4.245.1 : Bug fix, Missing null check on rowData.

diff --git a/Engine/Source/Programs/CSVTools/PerfReportTool/PerfReportTool.cs b/Engine/Source/Programs/CSVTools/PerfReportTool/PerfReportTool.cs
--- a/Engine/Source/Programs/CSVTools/PerfReportTool/PerfReportTool.cs
+++ b/Engine/Source/Programs/CSVTools/PerfReportTool/PerfReportTool.cs
@@ -23,7 +23,7 @@ namespace PerfReportTool
     class Version
     {
 		// Format: Major.Minor.Bugfix
-        private static string VersionString = "4.245.0";
+        private static string VersionString = "4.245.1";
 
         public static string Get() { return VersionString; }
     };
diff --git a/Engine/Source/Programs/CSVTools/PerfReportTool/Summaries/CheckpointSummary.cs b/Engine/Source/Programs/CSVTools/PerfReportTool/Summaries/CheckpointSummary.cs
--- a/Engine/Source/Programs/CSVTools/PerfReportTool/Summaries/CheckpointSummary.cs
+++ b/Engine/Source/Programs/CSVTools/PerfReportTool/Summaries/CheckpointSummary.cs
@@ -167,7 +167,7 @@ namespace PerfSummaries
 					}
 
 					// Requires capture event metricSuffix and summary metricPrefix to be something, otherwise no metric data will be written
-					if (metricPrefix.Length > 0 && capturePoint.metricSuffix.Length > 0)
+					if (rowData != null && metricPrefix.Length > 0 && capturePoint.metricSuffix.Length > 0)
 					{
 						rowData.Add(SummaryTableElement.Type.SummaryTableMetric, metricPrefix + "_" + baseStatName + "_" + capturePoint.metricSuffix, (double)value);
 					}
