PRT - add a zero-width space after slashes in stat names for all stat types

#rnx

diff --git a/Engine/Source/Programs/CSVTools/PerfReportTool/SummaryTable.cs b/Engine/Source/Programs/CSVTools/PerfReportTool/SummaryTable.cs
--- a/Engine/Source/Programs/CSVTools/PerfReportTool/SummaryTable.cs
+++ b/Engine/Source/Programs/CSVTools/PerfReportTool/SummaryTable.cs
@@ -649,22 +649,20 @@ namespace PerfSummaries
 				}
 			}
 
-			if (elementType == SummaryTableElement.Type.CsvStatAverage )
+			if (bGreyOutStatCategories)
 			{
-				if (bGreyOutStatCategories)
+				string baseStatName = SummaryTable.GetBaseStatNameWithPrefixAndSuffix(statName, out string prefix, out _);
+				int idx = baseStatName.LastIndexOf("/");
+				if (idx >= 0)
 				{
-					string baseStatName = SummaryTable.GetBaseStatNameWithPrefixAndSuffix(statName, out string prefix, out _);
-					int idx = baseStatName.LastIndexOf("/");
-					if (idx >= 0)
-					{
-						statName = prefix + "<span class='greyText'>" + baseStatName.Substring(0, idx+1) + "</span><span class='blackText'>" + baseStatName.Substring(idx+1)+ "</span>";
-					}
-				}
-				if (bAddStatCategorySeparatorSpaces)
-				{
-					return statName.Replace("/", "/ ");
+					statName = prefix + "<span class='greyText'>" + baseStatName.Substring(0, idx+1) + "</span><span class='blackText'>" + baseStatName.Substring(idx+1)+ "</span>";
 				}
 			}
+			if (bAddStatCategorySeparatorSpaces)
+			{
+				// Add a zero-width space to enable line breaks between slashes to reduce column width.
+				return statName.Replace("/", "/\u200B");
+			}
 			return statName;
 		}
 
