PRT - Revert zero width space in stat names to be a regular space.

#rnx

diff --git a/Engine/Source/Programs/CSVTools/PerfReportTool/PerfReportTool.cs b/Engine/Source/Programs/CSVTools/PerfReportTool/PerfReportTool.cs
--- a/Engine/Source/Programs/CSVTools/PerfReportTool/PerfReportTool.cs
+++ b/Engine/Source/Programs/CSVTools/PerfReportTool/PerfReportTool.cs
@@ -23,7 +23,7 @@ namespace PerfReportTool
     class Version
     {
 		// Format: Major.Minor.Bugfix
-        private static string VersionString = "4.247.0";
+        private static string VersionString = "4.247.1";
 
         public static string Get() { return VersionString; }
     };
diff --git a/Engine/Source/Programs/CSVTools/PerfReportTool/SummaryTable.cs b/Engine/Source/Programs/CSVTools/PerfReportTool/SummaryTable.cs
--- a/Engine/Source/Programs/CSVTools/PerfReportTool/SummaryTable.cs
+++ b/Engine/Source/Programs/CSVTools/PerfReportTool/SummaryTable.cs
@@ -660,8 +660,8 @@ namespace PerfSummaries
 			}
 			if (bAddStatCategorySeparatorSpaces)
 			{
-				// Add a zero-width space to enable line breaks between slashes to reduce column width.
-				return statName.Replace("/", "/\u200B");
+				// Add a space to enable line breaks between slashes to reduce column width.
+				return statName.Replace("/", "/ ");
 			}
 			return statName;
 		}
