PerfReportTool 4.248.0. Add support for summary table column sorting
The following modes are supported:
Default,            // Default sort (uses specified filter order, alphabetical wildcard resolution)
WildcardSortByAvg,  // When resolving wildcards, sort results by column average value
WildcardSortByMax,  // When resolving wildcards, sort results by column max value
SortByAvg,          // Sort all columns by average value (ignores specified filter order)
SortByMax,          // Sort all columns by max value (ignores specified filter order)
This can be specifed via the SummaryTable XML attribute "columnSortMode".
It can be overridden on the commandline via the "-columnSortMode <mode>" param.

diff --git a/Engine/Source/Programs/CSVTools/PerfReportTool/PerfReportTool.cs b/Engine/Source/Programs/CSVTools/PerfReportTool/PerfReportTool.cs
--- a/Engine/Source/Programs/CSVTools/PerfReportTool/PerfReportTool.cs
+++ b/Engine/Source/Programs/CSVTools/PerfReportTool/PerfReportTool.cs
@@ -23,7 +23,7 @@ namespace PerfReportTool
     class Version
     {
 		// Format: Major.Minor.Bugfix
-        private static string VersionString = "4.247.2";
+        private static string VersionString = "4.248.0";
 
         public static string Get() { return VersionString; }
     };
@@ -179,6 +179,7 @@ namespace PerfReportTool
 			"	   these resolve from existing metadata to form the final path. e.g D:\\ExtMeta\\{{platform}}\\{{buildversion}}.txt\n" +
 			"  -collatedStringVisibility <show|hide|auto> : sets visibility of strings in collated summary tables (default=auto)\n" +
 			"  -collatedDateVisibility <hide|newest|oldest> : sets visibility of dates in collated summary tables (default=newest)\n" +
+			"  -columnSortMode <default|sortByMax|sortByAvg|wildCardSortByMax|wildCardSortByAvg> : sorting mode for columns\n" +
 
 			"Diff rows\n" +
 			"  -addDiffRows : adds diff rows after the first two rows\n" +
@@ -891,6 +892,7 @@ namespace PerfReportTool
 			}
 			colorizeMode = GetEnumArg("colorizeTable", colorizeMode);
 
+			TableColumnSortMode tableColumnSortMode = GetEnumArg("columnSortMode", tableInfo.columnSortMode);
 
 			bool addMinMaxColumns = !GetBoolArg("noSummaryMinMax") && !bTransposeCollatedSummaryTable;
 
@@ -906,7 +908,7 @@ namespace PerfReportTool
 			table.SetColumnFormatInfo(reportXML.columnFormatInfoList);
 
 			bool bSortTrailingDigitsAsNumeric = GetBoolArg("sortTrailingDigitsAsNumeric");
-			SummaryTable filteredTable = table.SortAndFilter(tableInfo.columnFilterList, tableInfo.rowSortList, bReverseTable, weightByColumnName, additionalColumnFilters, bSortTrailingDigitsAsNumeric);
+			SummaryTable filteredTable = table.SortAndFilter(tableInfo.columnFilterList, tableInfo.rowSortList, bReverseTable, weightByColumnName, additionalColumnFilters, bSortTrailingDigitsAsNumeric, tableColumnSortMode);
 			if (bCollated)
 			{
 				StringCollationVisibility stringCollationVisibility = GetEnumArg("collatedStringVisibility", tableInfo.stringCollationVisibility);
diff --git a/Engine/Source/Programs/CSVTools/PerfReportTool/SummaryTable.cs b/Engine/Source/Programs/CSVTools/PerfReportTool/SummaryTable.cs
--- a/Engine/Source/Programs/CSVTools/PerfReportTool/SummaryTable.cs
+++ b/Engine/Source/Programs/CSVTools/PerfReportTool/SummaryTable.cs
@@ -103,6 +103,8 @@ namespace PerfSummaries
 			tableColorizeMode = tableElement.GetSafeAttribute<TableColorizeMode>(variableMappings, "colorizeMode", tableColorizeMode);
 			dateCollationVisibility = tableElement.GetSafeAttribute<DateCollationVisibility>(variableMappings, "collatedDateVisibility", dateCollationVisibility);
 			stringCollationVisibility = tableElement.GetSafeAttribute<StringCollationVisibility>(variableMappings, "collatedStringVisibility", stringCollationVisibility);
+			columnSortMode = tableElement.GetSafeAttribute<TableColumnSortMode>(variableMappings, "columnSortMode", columnSortMode);
+
 
 			statThreshold = tableElement.GetSafeAttribute<float>(variableMappings, "statThreshold", 0.0f);
 			hideStatPrefix = tableElement.GetSafeAttribute<string>(variableMappings, "hideStatPrefix");
@@ -180,6 +182,7 @@ namespace PerfSummaries
 		public TableColorizeMode tableColorizeMode = TableColorizeMode.Budget;
 		public DateCollationVisibility dateCollationVisibility = DateCollationVisibility.Newest;
 		public StringCollationVisibility stringCollationVisibility = StringCollationVisibility.Auto;
+		public TableColumnSortMode columnSortMode = TableColumnSortMode.Default;
 
 		public float statThreshold;
 		public string hideStatPrefix = null;
@@ -242,6 +245,15 @@ namespace PerfSummaries
 		Auto,		// Auto calculate based on other values in the column.
 	};
 
+	enum TableColumnSortMode
+	{
+		Default,            // Default sort (uses specified filter order, alphabetical wildcard resolution)
+		WildcardSortByAvg,  // When resolving wildcards, sort results by column average value
+		WildcardSortByMax,  // When resolving wildcards, sort results by column max value
+		SortByAvg,          // Sort all columns by average value (ignores specified filter order)
+		SortByMax,          // Sort all columns by max value (ignores specified filter order)
+	};
+
 	enum AutoColorizeMode
 	{
 		Off,
@@ -481,6 +493,47 @@ namespace PerfSummaries
 			return value == double.MaxValue ? 0.0 : value;
 		}
 
+		public double GetMaxValue()
+		{
+			if (!isNumeric || doubleValues.Count == 0)
+			{
+				return double.MaxValue;
+			}
+			double maxValue = double.MinValue;
+			foreach (double value in doubleValues)
+			{
+				maxValue = Math.Max(value, maxValue);
+			}
+			return maxValue;
+		}
+
+		public double GetMinValue()
+		{
+			if (!isNumeric || doubleValues.Count == 0)
+			{
+				return double.MaxValue;
+			}
+			double minValue = double.MaxValue;
+			foreach (double value in doubleValues)
+			{
+				minValue = Math.Min(value, minValue);
+			}
+			return minValue;
+		}
+
+		public double GetAvgValue()
+		{
+			if (!isNumeric || doubleValues.Count == 0)
+			{
+				return double.MaxValue;
+			}
+			double total = 0.0;
+			foreach (double value in doubleValues)
+			{
+				total += value;
+			}
+			return total / (double)doubleValues.Count;
+		}
 
 
 		public void AddDiffRows(bool bIsFirstColumn, DiffRowFrequency inDiffRowFrequency, bool bShowOnlyDiffRows)
@@ -1618,7 +1671,7 @@ namespace PerfSummaries
 			return SummaryTableElement.Type.COUNT;
 		}
 
-		public SummaryTable SortAndFilter(List<string> columnFilterList, List<string> rowSortList, bool bReverseSort, string weightByColumnName, IEnumerable<ISummaryTableColumnFilter> additionalFilters = null, bool bSortTrailingDigitsAsNumeric = false)
+		public SummaryTable SortAndFilter(List<string> columnFilterList, List<string> rowSortList, bool bReverseSort, string weightByColumnName, IEnumerable<ISummaryTableColumnFilter> additionalFilters, bool bSortTrailingDigitsAsNumeric, TableColumnSortMode tableColumnSortMode)
 		{
 			SummaryTable newTable = SortRows(rowSortList, bReverseSort, bSortTrailingDigitsAsNumeric);
 
@@ -1650,14 +1703,30 @@ namespace PerfSummaries
 					List<string> keyList = CsvStats.WildcardMatchStringList(allKeysList, unqualifiedKey, false, true);
 
 					// Resolve the keyList and output the columns 
+					List<SummaryTableColumn> foundColumns = new List<SummaryTableColumn>();
 					foreach (string key in keyList)
 					{
 						SummaryTableColumn column = columnLookup.Get(key, statType);
 						if (column != null)
 						{
-							newColumnList.Add(column);
+							foundColumns.Add(column);
 						}
 					}
+					// Apply sorting to wildcard columns if requested
+					if (tableColumnSortMode == TableColumnSortMode.WildcardSortByMax)
+					{
+						foundColumns.Sort((a, b) => b.GetMaxValue().CompareTo(a.GetMaxValue()));
+					}
+					else if (tableColumnSortMode == TableColumnSortMode.WildcardSortByAvg)
+					{
+						foundColumns.Sort((a, b) => b.GetAvgValue().CompareTo(a.GetAvgValue()));
+					}
+
+					foreach (SummaryTableColumn column in foundColumns)
+					{
+						newColumnList.Add(column);
+					}
+
 				}
 				else
 				{
@@ -1672,6 +1741,16 @@ namespace PerfSummaries
 			// Remove duplicates
 			newColumnList = newColumnList.Distinct().ToList();
 
+			// Apply sorting to all columns if requested
+			if (tableColumnSortMode == TableColumnSortMode.SortByMax)
+			{
+				newColumnList.Sort((a, b) => b.GetMaxValue().CompareTo(a.GetMaxValue()));
+			}
+			else if (tableColumnSortMode == TableColumnSortMode.SortByAvg)
+			{
+				newColumnList.Sort((a, b) => b.GetAvgValue().CompareTo(a.GetAvgValue()));
+			}
+
 			// Compute row weights
 			if (weightByColumnName != null && columnLookup.ContainsKey(weightByColumnName))
 			{
