PerfReportTool: Handle missing values correctly when sorting columns

diff --git a/Engine/Source/Programs/CSVTools/PerfReportTool/SummaryTable.cs b/Engine/Source/Programs/CSVTools/PerfReportTool/SummaryTable.cs
--- a/Engine/Source/Programs/CSVTools/PerfReportTool/SummaryTable.cs
+++ b/Engine/Source/Programs/CSVTools/PerfReportTool/SummaryTable.cs
@@ -495,44 +495,71 @@ namespace PerfSummaries
 
 		public double GetMaxValue()
 		{
-			if (!isNumeric || doubleValues.Count == 0)
+			if (!isNumeric)
 			{
 				return double.MaxValue;
 			}
 			double maxValue = double.MinValue;
+			int count = 0;
 			foreach (double value in doubleValues)
 			{
-				maxValue = Math.Max(value, maxValue);
+				if (value != double.MaxValue)
+				{
+					maxValue = Math.Max(value, maxValue);
+					count++;
+				}
+			}
+			if (count == 0)
+			{
+				return 0.0;
 			}
 			return maxValue;
 		}
 
 		public double GetMinValue()
 		{
-			if (!isNumeric || doubleValues.Count == 0)
+			if (!isNumeric)
 			{
 				return double.MaxValue;
 			}
 			double minValue = double.MaxValue;
+			int count = 0;
 			foreach (double value in doubleValues)
 			{
-				minValue = Math.Min(value, minValue);
+				if (value != double.MaxValue)
+				{
+					minValue = Math.Min(value, minValue);
+					count++;
+				}
+			}
+			if (count == 0)
+			{
+				return 0.0;
 			}
 			return minValue;
 		}
 
 		public double GetAvgValue()
 		{
-			if (!isNumeric || doubleValues.Count == 0)
+			if (!isNumeric)
 			{
 				return double.MaxValue;
 			}
 			double total = 0.0;
+			int count = 0;
 			foreach (double value in doubleValues)
 			{
-				total += value;
+				if (value != double.MaxValue)
+				{
+					total += value;
+					count++;
+				}
+			}
+			if (count == 0)
+			{
+				return 0.0;
 			}
-			return total / (double)doubleValues.Count;
+			return total / (double)count;
 		}
 
 
