PerfReportTool: Fixing PRC to json failing when encountering NaNs

diff --git a/Engine/Source/Programs/CSVTools/PerfReportTool/SummaryTableData.cs b/Engine/Source/Programs/CSVTools/PerfReportTool/SummaryTableData.cs
--- a/Engine/Source/Programs/CSVTools/PerfReportTool/SummaryTableData.cs
+++ b/Engine/Source/Programs/CSVTools/PerfReportTool/SummaryTableData.cs
@@ -62,7 +62,7 @@ namespace PerfSummaries
 
 		private static void WriteJsonFile(bool bSerializeToStream, string filename, Dictionary<string, dynamic> dictToSerialize, bool bWriteIndented)
 		{
-			JsonSerializerOptions options = new JsonSerializerOptions { WriteIndented = bWriteIndented };
+			JsonSerializerOptions options = new JsonSerializerOptions { WriteIndented = bWriteIndented, NumberHandling = System.Text.Json.Serialization.JsonNumberHandling.AllowNamedFloatingPointLiterals };
 
 			if (bSerializeToStream)
 			{
@@ -285,8 +285,17 @@ namespace PerfSummaries
 			}
 			if (isNumeric)
 			{
-				// Serialize as Decimal to avoid 0.1 being serialized as 0.1000000000000001 etc
-				Dict.Add("value", new Decimal(numericValue));
+				try
+				{
+					// Serialize as Decimal to avoid 0.1 being serialized as 0.1000000000000001 etc
+					Dict.Add("value", new Decimal(numericValue));
+				}
+				catch (OverflowException)
+				{
+					// if Decimal constructor threw then this is a value it doesn't like such as NaN, inf, etc. In those
+					// cases just pass through the value as a double
+					Dict.Add("value", numericValue);
+				}
 			}
 			else
 			{
