Fixed issue with duplicate SchemaVersions being added to events when using CSVToTelemetry
Removed the a redundant duplicate attribute check from AnalyticsProviderBroadcast as this will be called again in the IAnalyticsProvider interface
Added duplicate attribute check for default attributes against the provided attributes in the IAnalyticsPorivder RecordEvent method
Moved the check for a live Telemetry session outside of the check for a valid input file. If there is no valid AnalyticsProvider to receive it then we should neve process or send any event data.

diff --git a/Engine/Source/Programs/CSVTools/CSVToTelemetry/Private/CSVToTelemetry.cpp b/Engine/Source/Programs/CSVTools/CSVToTelemetry/Private/CSVToTelemetry.cpp
--- a/Engine/Source/Programs/CSVTools/CSVToTelemetry/Private/CSVToTelemetry.cpp
+++ b/Engine/Source/Programs/CSVTools/CSVToTelemetry/Private/CSVToTelemetry.cpp
@@ -25,20 +25,20 @@ static int GenerateTelemetryFromCSVFile(const FString& FilePath)
 		return Result;
 	}
 
-	if (FilePath.Len() > 0)
-	{
-		TArray<FString> Rows;
+	FStudioTelemetry::Get().StartSession();
 
-		// Load all the rows from the file into memory
-		if (FFileHelper::LoadFileToStringArray(Rows, *FilePath))
+	// Start the telemetry session and record each row as a single telemetry event 
+	if (FStudioTelemetry::Get().IsSessionRunning())
+	{
+		if (FilePath.Len() > 0)
 		{
-			UE_LOG(LogCSVToTelemetry, Display, TEXT("Imported %d rows from file"), Rows.Num());
-
-			FStudioTelemetry::Get().StartSession();
+			TArray<FString> Rows;
 
-			// Start the telemetry session and record each row as a single telemetry event 
-			if (FStudioTelemetry::Get().IsSessionRunning())
+			// Load all the rows from the file into memory
+			if (FFileHelper::LoadFileToStringArray(Rows, *FilePath))
 			{
+				UE_LOG(LogCSVToTelemetry, Display, TEXT("Imported %d rows from file"), Rows.Num());
+
 				uint32 TotalUploadSize = 0;
 				uint32 RowIndex = 0;
 				uint32 KeySize = 0;
@@ -50,16 +50,11 @@ static int GenerateTelemetryFromCSVFile(const FString& FilePath)
 
 				int32 SchemaVersion(-1);
 
-				TSharedPtr<IAnalyticsProvider> Provider = FStudioTelemetry::Get().GetProvider().Pin();
-				TArray<FAnalyticsEventAttribute> DefaultAttributes = Provider->GetDefaultEventAttributesSafe();
-
-				if (FParse::Value(FCommandLine::Get(), TEXT("schema="), SchemaVersion) == true)
+				if (FParse::Value(FCommandLine::Get(), TEXT("schema="), SchemaVersion) == false)
 				{
-					DefaultAttributes.Emplace(TEXT("SchemaVersion"), SchemaVersion);
+					SchemaVersion = -1;
 				}
-
-				Provider->SetDefaultEventAttributes(MoveTemp(DefaultAttributes));
-
+				
 				if (FParse::Value(FCommandLine::Get(), TEXT("columns="), Columns) == true)
 				{
 					// We have specified the CSV columns on the command line 
@@ -93,6 +88,12 @@ static int GenerateTelemetryFromCSVFile(const FString& FilePath)
 
 							TArray<FAnalyticsEventAttribute> Attributes;
 
+							if (SchemaVersion != -1)
+							{
+								// Add our schema version
+								Attributes.Emplace(TEXT("SchemaVersion"), SchemaVersion);
+							}
+							
 							// Add the Key/Value pair to the attribute list
 							for (int32 i = 0; i < ValueArray.Num(); ++i)
 							{
@@ -118,8 +119,6 @@ static int GenerateTelemetryFromCSVFile(const FString& FilePath)
 
 				UE_LOG(LogCSVToTelemetry, Display, TEXT("Generated %d bytes of event data"), TotalUploadSize);
 
-				FStudioTelemetry::Get().EndSession();
-
 				Result = 0;
 			}
 		}
@@ -127,6 +126,12 @@ static int GenerateTelemetryFromCSVFile(const FString& FilePath)
 		{
 			UE_LOG(LogCSVToTelemetry, Error, TEXT("Unable to read rows from CSV file %s"), *FilePath);
 		}
+
+		FStudioTelemetry::Get().EndSession();
+	}
+	else
+	{
+		UE_LOG(LogCSVToTelemetry, Error, TEXT("Unable to start a telemetry session"));
 	}
 
 	return Result;
diff --git a/Engine/Source/Runtime/Analytics/Analytics/Private/Interfaces/IAnalyticsProvider.cpp b/Engine/Source/Runtime/Analytics/Analytics/Private/Interfaces/IAnalyticsProvider.cpp
--- a/Engine/Source/Runtime/Analytics/Analytics/Private/Interfaces/IAnalyticsProvider.cpp
+++ b/Engine/Source/Runtime/Analytics/Analytics/Private/Interfaces/IAnalyticsProvider.cpp
@@ -3,17 +3,19 @@
 #include "Interfaces/IAnalyticsProvider.h"
 #include "Analytics.h"
 
-void IAnalyticsProvider::CheckForDuplicateAttributes(const FString& EventName, const TArray<FAnalyticsEventAttribute>& Attributes)
+void IAnalyticsProvider::CheckForDuplicateAttributes(const FString& EventName, const TArray<FAnalyticsEventAttribute>& AttributesA, const TArray<FAnalyticsEventAttribute>& AttributesB)
 {
 #if UE_BUILD_DEBUG || UE_BUILD_DEVELOPMENT
-	// Expose events that have duplicate attribute names. This is is not handled by the analytics backends in any reliable way.
-	for (int32 index0 = 0; index0 < Attributes.Num(); ++index0)
+	// Expose events that have duplicate attribute names. Not only does this provide ambiguous event data, in general this is not handled by all the analytics backends in any consistent manner.
+	for (int32 indexA = 0; indexA < AttributesA.Num(); ++indexA)
 	{
-		for (int32 index1 = index0 + 1; index1 < Attributes.Num(); ++index1)
+		const FString& AttributeName = AttributesA[indexA].GetName();
+
+		for (int32 indexB = indexA + 1; indexB < AttributesB.Num(); ++indexB)
 		{
-			if (Attributes[index0].GetName() == Attributes[index1].GetName())
+			if (AttributeName == AttributesB[indexB].GetName())
 			{
-				UE_LOG(LogAnalytics, Warning, TEXT("Duplicate Attributes Found For Event %s %s==%s"), *EventName, *Attributes[index0].GetName(), *Attributes[index1].GetName());
+				UE_LOG(LogAnalytics, Warning, TEXT("Duplicate Attributes %s Found For Event %s"), *AttributeName, *EventName);
 			}
 		}
 	}
diff --git a/Engine/Source/Runtime/Analytics/Analytics/Public/Interfaces/IAnalyticsProvider.h b/Engine/Source/Runtime/Analytics/Analytics/Public/Interfaces/IAnalyticsProvider.h
--- a/Engine/Source/Runtime/Analytics/Analytics/Public/Interfaces/IAnalyticsProvider.h
+++ b/Engine/Source/Runtime/Analytics/Analytics/Public/Interfaces/IAnalyticsProvider.h
@@ -450,7 +450,7 @@ protected:
 	/**
 	 * Checks for duplicate key names in the attribute array. Only implemented for debug and development builds.
 	 */
-	void ANALYTICS_API CheckForDuplicateAttributes(const FString& EventName, const TArray<FAnalyticsEventAttribute>& Attributes);
+	void ANALYTICS_API CheckForDuplicateAttributes(const FString& EventName, const TArray<FAnalyticsEventAttribute>& AttributesA, const TArray<FAnalyticsEventAttribute>& AttributesB);
 
 };
 
diff --git a/Engine/Source/Runtime/Analytics/AnalyticsET/Private/IAnalyticsProviderET.cpp b/Engine/Source/Runtime/Analytics/AnalyticsET/Private/IAnalyticsProviderET.cpp
--- a/Engine/Source/Runtime/Analytics/AnalyticsET/Private/IAnalyticsProviderET.cpp
+++ b/Engine/Source/Runtime/Analytics/AnalyticsET/Private/IAnalyticsProviderET.cpp
@@ -619,7 +619,11 @@ void FAnalyticsProviderET::RecordEvent(FString&& EventName, const TArray<FAnalyt
 {
 	LLM_SCOPE_BYNAME(TEXT("Analytics"));
 
-	CheckForDuplicateAttributes(EventName, Attributes);
+	// Check for duplicate attributes against themselves
+	CheckForDuplicateAttributes(EventName, Attributes, Attributes);
+
+	// Check for duplicate attributes against the defaults
+	CheckForDuplicateAttributes(EventName, Attributes, GetDefaultEventAttributesSafe());
 
 	// let higher level code filter the decision of whether to send the event
 	if (ShouldRecordEvent(EventName))
diff --git a/Engine/Source/Runtime/StudioTelemetry/Private/AnalyticsProviderBroadcast.cpp b/Engine/Source/Runtime/StudioTelemetry/Private/AnalyticsProviderBroadcast.cpp
--- a/Engine/Source/Runtime/StudioTelemetry/Private/AnalyticsProviderBroadcast.cpp
+++ b/Engine/Source/Runtime/StudioTelemetry/Private/AnalyticsProviderBroadcast.cpp
@@ -261,8 +261,6 @@ void FAnalyticsProviderBroadcast::SetRecordEventCallback(OnRecordEvent Callback)
 
 void FAnalyticsProviderBroadcast::RecordEvent(const FString& EventName, const TArray<FAnalyticsEventAttribute>& Attributes)
 {
-	CheckForDuplicateAttributes(EventName, Attributes);
-
 	for (TProviders::TConstIterator it(Providers); it; ++it)
 	{
 		(*it).Value->RecordEvent(EventName, Attributes);
diff --git a/Engine/Source/Runtime/StudioTelemetry/Private/StudioTelemetry.cpp b/Engine/Source/Runtime/StudioTelemetry/Private/StudioTelemetry.cpp
--- a/Engine/Source/Runtime/StudioTelemetry/Private/StudioTelemetry.cpp
+++ b/Engine/Source/Runtime/StudioTelemetry/Private/StudioTelemetry.cpp
@@ -186,7 +186,7 @@ void FStudioTelemetry::StartSession()
 
 		TArray<FAnalyticsEventAttribute> Attributes;
 		Attributes.Emplace(TEXT("SchemaVersion"), 1);
-		AnalyticsProvider->RecordEvent(TEXT("StudioTelemetry.SessionStart"), TArray<FAnalyticsEventAttribute>());
+		AnalyticsProvider->RecordEvent(TEXT("StudioTelemetry.SessionStart"), Attributes);
 
 		OnStartSession.Broadcast();
 
@@ -202,7 +202,9 @@ void FStudioTelemetry::EndSession()
 
 		OnEndSession.Broadcast();
 
-		AnalyticsProvider->RecordEvent(TEXT("StudioTelemetry.SessionEnd"), TArray<FAnalyticsEventAttribute>());
+		TArray<FAnalyticsEventAttribute> Attributes;
+		Attributes.Emplace(TEXT("SchemaVersion"), 1);
+		AnalyticsProvider->RecordEvent(TEXT("StudioTelemetry.SessionEnd"), Attributes);
 
 		// End session for the tracer and the provider
 		if (AnalyticsTracer.IsValid())
@@ -224,7 +226,7 @@ void FStudioTelemetry::EndSession()
 
 bool FStudioTelemetry::IsSessionRunning() const
 {
-	return AnalyticsProvider.IsValid();
+	return AnalyticsProvider.IsValid() && AnalyticsProvider->HasValidProviders();
 }
 
 void FStudioTelemetry::LoadConfiguration()
