Fix for potential out of bounds issue highlighted by linter and incorrect operation of the check duplicate attributes function

diff --git a/Engine/Source/Programs/CSVTools/CSVToTelemetry/Private/CSVToTelemetry.cpp b/Engine/Source/Programs/CSVTools/CSVToTelemetry/Private/CSVToTelemetry.cpp
--- a/Engine/Source/Programs/CSVTools/CSVToTelemetry/Private/CSVToTelemetry.cpp
+++ b/Engine/Source/Programs/CSVTools/CSVToTelemetry/Private/CSVToTelemetry.cpp
@@ -15,22 +15,22 @@ IMPLEMENT_APPLICATION(CSVToTelemetry, "CSVToTelemetry");
 
 static int GenerateTelemetryFromCSVFile(const FString& FilePath)
 {
-	int32 Result = 1;
-
-	FString EventName;
+	int32 Result = 0;
 
-	if (FParse::Value(FCommandLine::Get(), TEXT("event="), EventName) == false)
+	if (FilePath.Len() > 0)
 	{
-		UE_LOG(LogCSVToTelemetry, Error, TEXT("Must provide a Event name with -event=[eventname]"));
-		return Result;
-	}
+		FString EventName;
+
+		if (FParse::Value(FCommandLine::Get(), TEXT("event="), EventName) == false)
+		{
+			UE_LOG(LogCSVToTelemetry, Error, TEXT("Must provide a Event name with -event=[eventname]"));
+			return Result;
+		}
 
-	FStudioTelemetry::Get().StartSession();
+		FStudioTelemetry::Get().StartSession();
 
-	// Start the telemetry session and record each row as a single telemetry event 
-	if (FStudioTelemetry::Get().IsSessionRunning())
-	{
-		if (FilePath.Len() > 0)
+		// Start the telemetry session and record each row as a single telemetry event 
+		if (FStudioTelemetry::Get().IsSessionRunning())
 		{
 			TArray<FString> Rows;
 
@@ -54,7 +54,7 @@ static int GenerateTelemetryFromCSVFile(const FString& FilePath)
 				{
 					SchemaVersion = -1;
 				}
-				
+
 				if (FParse::Value(FCommandLine::Get(), TEXT("columns="), Columns) == true)
 				{
 					// We have specified the CSV columns on the command line 
@@ -93,7 +93,7 @@ static int GenerateTelemetryFromCSVFile(const FString& FilePath)
 								// Add our schema version
 								Attributes.Emplace(TEXT("SchemaVersion"), SchemaVersion);
 							}
-							
+
 							// Add the Key/Value pair to the attribute list
 							for (int32 i = 0; i < ValueArray.Num(); ++i)
 							{
@@ -117,21 +117,21 @@ static int GenerateTelemetryFromCSVFile(const FString& FilePath)
 					RowIndex++;
 				}
 
-				UE_LOG(LogCSVToTelemetry, Display, TEXT("Generated %d bytes of event data"), TotalUploadSize);
+				Result = 1;
 
-				Result = 0;
+				UE_LOG(LogCSVToTelemetry, Display, TEXT("Generated %d bytes of event data"), TotalUploadSize);
+			}
+			else
+			{
+				UE_LOG(LogCSVToTelemetry, Error, TEXT("Unable to read rows from CSV file %s"), *FilePath);
 			}
+
+			FStudioTelemetry::Get().EndSession();
 		}
 		else
 		{
-			UE_LOG(LogCSVToTelemetry, Error, TEXT("Unable to read rows from CSV file %s"), *FilePath);
+			UE_LOG(LogCSVToTelemetry, Error, TEXT("Unable to start a telemetry session"));
 		}
-
-		FStudioTelemetry::Get().EndSession();
-	}
-	else
-	{
-		UE_LOG(LogCSVToTelemetry, Error, TEXT("Unable to start a telemetry session"));
 	}
 
 	return Result;
diff --git a/Engine/Source/Runtime/Analytics/Analytics/Private/Interfaces/IAnalyticsProvider.cpp b/Engine/Source/Runtime/Analytics/Analytics/Private/Interfaces/IAnalyticsProvider.cpp
--- a/Engine/Source/Runtime/Analytics/Analytics/Private/Interfaces/IAnalyticsProvider.cpp
+++ b/Engine/Source/Runtime/Analytics/Analytics/Private/Interfaces/IAnalyticsProvider.cpp
@@ -3,6 +3,25 @@
 #include "Interfaces/IAnalyticsProvider.h"
 #include "Analytics.h"
 
+void IAnalyticsProvider::CheckForDuplicateAttributes(const FString& EventName, const TArray<FAnalyticsEventAttribute>& Attributes)
+{
+#if UE_BUILD_DEBUG || UE_BUILD_DEVELOPMENT
+	// Expose events that have duplicate attribute names. Not only does this provide ambiguous event data, in general this is not handled by all the analytics backends in any consistent manner.
+	for (int32 indexA = 0; indexA < Attributes.Num(); ++indexA)
+	{
+		const FString& AttributeName = Attributes[indexA].GetName();
+
+		for (int32 indexB = indexA+1; indexB < Attributes.Num(); ++indexB)
+		{
+			if (AttributeName == Attributes[indexB].GetName())
+			{
+				UE_LOG(LogAnalytics, Warning, TEXT("Duplicate Attributes %s Found For Event %s"), *AttributeName, *EventName);
+			}
+		}
+	}
+#endif
+}
+
 void IAnalyticsProvider::CheckForDuplicateAttributes(const FString& EventName, const TArray<FAnalyticsEventAttribute>& AttributesA, const TArray<FAnalyticsEventAttribute>& AttributesB)
 {
 #if UE_BUILD_DEBUG || UE_BUILD_DEVELOPMENT
@@ -11,7 +30,7 @@ void IAnalyticsProvider::CheckForDuplicateAttributes(const FString& EventName, c
 	{
 		const FString& AttributeName = AttributesA[indexA].GetName();
 
-		for (int32 indexB = indexA + 1; indexB < AttributesB.Num(); ++indexB)
+		for (int32 indexB = 0; indexB < AttributesB.Num(); ++indexB)
 		{
 			if (AttributeName == AttributesB[indexB].GetName())
 			{
diff --git a/Engine/Source/Runtime/Analytics/Analytics/Public/Interfaces/IAnalyticsProvider.h b/Engine/Source/Runtime/Analytics/Analytics/Public/Interfaces/IAnalyticsProvider.h
--- a/Engine/Source/Runtime/Analytics/Analytics/Public/Interfaces/IAnalyticsProvider.h
+++ b/Engine/Source/Runtime/Analytics/Analytics/Public/Interfaces/IAnalyticsProvider.h
@@ -448,7 +448,12 @@ public:
 
 protected:
 	/**
-	 * Checks for duplicate key names in the attribute array. Only implemented for debug and development builds.
+	 * Checks for duplicate key names in the same attribute array. Only implemented for debug and development builds.
+	 */
+	void ANALYTICS_API CheckForDuplicateAttributes(const FString& EventName, const TArray<FAnalyticsEventAttribute>& Attributes);
+
+	/**
+	 * Checks for duplicate key names between two different attribute arrays. Only implemented for debug and development builds.
 	 */
 	void ANALYTICS_API CheckForDuplicateAttributes(const FString& EventName, const TArray<FAnalyticsEventAttribute>& AttributesA, const TArray<FAnalyticsEventAttribute>& AttributesB);
 
diff --git a/Engine/Source/Runtime/Analytics/AnalyticsET/Private/IAnalyticsProviderET.cpp b/Engine/Source/Runtime/Analytics/AnalyticsET/Private/IAnalyticsProviderET.cpp
--- a/Engine/Source/Runtime/Analytics/AnalyticsET/Private/IAnalyticsProviderET.cpp
+++ b/Engine/Source/Runtime/Analytics/AnalyticsET/Private/IAnalyticsProviderET.cpp
@@ -620,7 +620,7 @@ void FAnalyticsProviderET::RecordEvent(FString&& EventName, const TArray<FAnalyt
 	LLM_SCOPE_BYNAME(TEXT("Analytics"));
 
 	// Check for duplicate attributes against themselves
-	CheckForDuplicateAttributes(EventName, Attributes, Attributes);
+	CheckForDuplicateAttributes(EventName, Attributes);
 
 	// Check for duplicate attributes against the defaults
 	CheckForDuplicateAttributes(EventName, Attributes, GetDefaultEventAttributesSafe());
