Fixed incorrect normal operation return value in CSVToTelemetry which resulted in errors on the build job.

diff --git a/Engine/Source/Programs/CSVTools/CSVToTelemetry/Private/CSVToTelemetry.cpp b/Engine/Source/Programs/CSVTools/CSVToTelemetry/Private/CSVToTelemetry.cpp
--- a/Engine/Source/Programs/CSVTools/CSVToTelemetry/Private/CSVToTelemetry.cpp
+++ b/Engine/Source/Programs/CSVTools/CSVToTelemetry/Private/CSVToTelemetry.cpp
@@ -15,7 +15,8 @@ IMPLEMENT_APPLICATION(CSVToTelemetry, "CSVToTelemetry");
 
 static int GenerateTelemetryFromCSVFile(const FString& FilePath)
 {
-	int32 Result = 0;
+	// Assume the result will be an error state, non-zero, and only returns zero when fully successful.
+	int32 Result = 1;
 
 	if (FilePath.Len() > 0)
 	{
@@ -117,9 +118,10 @@ static int GenerateTelemetryFromCSVFile(const FString& FilePath)
 					RowIndex++;
 				}
 
-				Result = 1;
-
 				UE_LOG(LogCSVToTelemetry, Display, TEXT("Generated %d bytes of event data"), TotalUploadSize);
+
+				// Everything worked as expected return code should be zero
+				Result = 0;
 			}
 			else
 			{
@@ -139,6 +141,7 @@ static int GenerateTelemetryFromCSVFile(const FString& FilePath)
 
 static int GenerateTelemetryFromCSVProfileFile(const FString& FilePath)
 {
+	// Assume the result will be an error state, non-zero, and only returns zero when fully successful.
 	int32 Result = 1;
 
 	FString EventName;
