PerfReportTool 4.50.1: Fix logspam due to duplicate keys with external metadata due to key case mismatch

diff --git a/Engine/Source/Programs/CSVTools/PerfReportTool/CsvFileCache..cs b/Engine/Source/Programs/CSVTools/PerfReportTool/CsvFileCache..cs
--- a/Engine/Source/Programs/CSVTools/PerfReportTool/CsvFileCache..cs
+++ b/Engine/Source/Programs/CSVTools/PerfReportTool/CsvFileCache..cs
@@ -257,9 +257,8 @@ namespace PerfReportTool
 				foreach (KeyValuePair<string, string> entry in fileDict)
 				{
 					// Add metadata if it doesn't already exist
-					SummaryTableElement element = summaryTableRowData.Get(entry.Key);
-					if (element == null)
-					{
+					if (!summaryTableRowData.Contains(entry.Key.ToLower()))
+					{ 
 						summaryTableRowData.Add(SummaryTableElement.Type.ExternalMetadata, entry.Key, entry.Value);
 					}
 				}
diff --git a/Engine/Source/Programs/CSVTools/PerfReportTool/PerfReportTool.cs b/Engine/Source/Programs/CSVTools/PerfReportTool/PerfReportTool.cs
--- a/Engine/Source/Programs/CSVTools/PerfReportTool/PerfReportTool.cs
+++ b/Engine/Source/Programs/CSVTools/PerfReportTool/PerfReportTool.cs
@@ -23,7 +23,7 @@ namespace PerfReportTool
     class Version
     {
 		// Format: Major.Minor.Bugfix
-        private static string VersionString = "4.250.0";
+        private static string VersionString = "4.250.1";
 
         public static string Get() { return VersionString; }
     };
