PerfReportTool 4.251.0: improved summary table web interop support
- Add support for outputting embedded JS data describing summary table rows and columns. Pass -summaryTableJsData to enable
- Add row_<I> and col_<J> ids to <tr> and header <th>/<td> elements
- Add support for hiding individual columns so they're not output as HTML
- Add a hidden csvids column so we can track a row's csv IDs after collation
#tests PRT unit tests
[RN] PerfReportTool 4.251.0: Improved summary table web interop support. Add support for outputting embedded JS data describing summary table rows and columns. Pass -summaryTableJsData to enable

diff --git a/Engine/Source/Programs/CSVTools/PerfReportTool/PerfReportTool.cs b/Engine/Source/Programs/CSVTools/PerfReportTool/PerfReportTool.cs
--- a/Engine/Source/Programs/CSVTools/PerfReportTool/PerfReportTool.cs
+++ b/Engine/Source/Programs/CSVTools/PerfReportTool/PerfReportTool.cs
@@ -23,7 +23,7 @@ namespace PerfReportTool
     class Version
     {
 		// Format: Major.Minor.Bugfix
-        private static string VersionString = "4.250.1";
+        private static string VersionString = "4.251.0";
 
         public static string Get() { return VersionString; }
     };
@@ -180,6 +180,7 @@ namespace PerfReportTool
 			"  -collatedStringVisibility <show|hide|auto> : sets visibility of strings in collated summary tables (default=auto)\n" +
 			"  -collatedDateVisibility <hide|newest|oldest> : sets visibility of dates in collated summary tables (default=newest)\n" +
 			"  -columnSortMode <default|sortByMax|sortByAvg|wildCardSortByMax|wildCardSortByAvg> : sorting mode for columns\n" +
+			"  -summaryTableJsData : embeds javascript row and column data in the <script> section to help with interop\n" +
 
 			"Diff rows\n" +
 			"  -addDiffRows : adds diff rows after the first two rows\n" +
@@ -945,7 +946,8 @@ namespace PerfReportTool
 					weightByColumnName, 
 					summaryTitle,
 					bCollated ? bTransposeCollatedSummaryTable : bTransposeFullSummaryTable,
-					summaryTableStyle: tableInfo.summaryTableStyle
+					tableInfo.summaryTableStyle,
+					GetBoolArg("summaryTableJsData")
 				);
 			}
 		}
diff --git a/Engine/Source/Programs/CSVTools/PerfReportTool/SummaryTable.cs b/Engine/Source/Programs/CSVTools/PerfReportTool/SummaryTable.cs
--- a/Engine/Source/Programs/CSVTools/PerfReportTool/SummaryTable.cs
+++ b/Engine/Source/Programs/CSVTools/PerfReportTool/SummaryTable.cs
@@ -1,12 +1,14 @@
 // Copyright Epic Games, Inc. All Rights Reserved.
 
+using CSVStats;
+using PerfReportTool;
 using System;
 using System.Collections.Generic;
+using System.Data.Common;
 using System.Linq;
 using System.Text;
+using System.Text.Json;
 using System.Xml.Linq;
-using CSVStats;
-using PerfReportTool;
 
 namespace PerfSummaries
 {
@@ -403,6 +405,7 @@ namespace PerfSummaries
 		public DiffRowFrequency diffRowFrequency = DiffRowFrequency.None;
 		public bool bShowOnlyDiffRows = false;
 		public bool isCountColumn = false;
+		public bool isSortByColumn = false;
 		// Column header tooltip. Displayed when hovering over the header.
 		public string tooltip = null;
 		// Multiplied with the colour of each cell to modify the colour (including header).
@@ -420,6 +423,7 @@ namespace PerfSummaries
 
 		List<ColourThresholdList> colourThresholds = new List<ColourThresholdList>();
 		ColourThresholdList colourThresholdOverride = null;
+		public bool isVisible = true;
 
 		public SummaryTableColumn(
 			string inName,
@@ -432,7 +436,9 @@ namespace PerfSummaries
 			Colour inColumnColourModifier = null,
 			ColumnAggregateType inAggregateType = ColumnAggregateType.None,
 			SummaryTableColumn inAggregateBaseColumn = null,
-			bool bInIsCountColumn = false)
+			bool bInIsCountColumn = false,
+			bool bInIsSortByColumn = false,
+			bool bInIsVisible = true )
 		{
 			name = SummaryTableColumn.getAggregateTypePrefix(inAggregateType) + inName;
 			isNumeric = inIsNumeric;
@@ -445,6 +451,8 @@ namespace PerfSummaries
 			aggregateType = inAggregateType;
 			aggregateBaseColumn = inAggregateBaseColumn;
 			isCountColumn = bInIsCountColumn;
+			isSortByColumn = bInIsSortByColumn;
+			isVisible = bInIsVisible;
 			if (inAggregateType == ColumnAggregateType.Avg && aggregateBaseColumn == null)
 			{
 				aggregateBaseColumn = this;
@@ -471,9 +479,25 @@ namespace PerfSummaries
 			return bIncludeTypeQualifier ? SummaryTable.GetElementTypeStatPrefix(this.elementType) + name.ToLower() : name.ToLower();
 		}
 
+
+		public string GetJsData()
+		{
+			string baseName = SummaryTable.GetBaseStatNameWithPrefixAndSuffix(name, out _, out string suffix);
+			Dictionary<string, string> attributes = new Dictionary<string, string>();
+			// Strip the column prefix (min/max/avg) since we're outputting the aggregate column types separately anyway
+			// The suffix is included (e.g MemoryFreeMB Min)
+			attributes.Add("name", baseName + suffix);
+			if (aggregateType != ColumnAggregateType.None)
+			{
+				attributes.Add("aggregateType", aggregateType.ToString());
+			}
+			attributes.Add("elementType", elementType.ToString());
+			return JsonSerializer.Serialize(attributes);
+		}
+
 		public SummaryTableColumn Clone()
 		{
-			SummaryTableColumn newColumn = new SummaryTableColumn(name, isNumeric, displayName, isRowWeightColumn, elementType, formatInfo, tooltip, columnColourModifier);
+			SummaryTableColumn newColumn = new SummaryTableColumn(name, isNumeric, displayName, isRowWeightColumn, elementType, formatInfo, tooltip, columnColourModifier, aggregateType, aggregateBaseColumn, isCountColumn, isSortByColumn, isVisible );
 			newColumn.doubleValues.AddRange(doubleValues);
 			newColumn.stringValues.AddRange(stringValues);
 			newColumn.colourThresholds.AddRange(colourThresholds);
@@ -1348,14 +1372,14 @@ namespace PerfSummaries
 			return sb.ToString();
 		}
 
-		private SummaryTableColumn CollateColumn(SummaryTableColumn srcColumn, ColumnAggregateType aggregateType, List<int> collatedRowSourceRowCounts, SummaryTableColumn aggregateBaseColumn = null)
+		private SummaryTableColumn CollateColumn(SummaryTableColumn srcColumn, ColumnAggregateType aggregateType, List<int> collatedRowSourceRowCounts, SummaryTableColumn aggregateBaseColumn = null, bool bShowMergedStringValues = false)
 		{
-			return CollateColumn(srcColumn, aggregateType, collatedRowSourceRowCounts, aggregateBaseColumn, out int temp);
+			return CollateColumn(srcColumn, aggregateType, collatedRowSourceRowCounts, aggregateBaseColumn, bShowMergedStringValues, out int temp);
 		}
 
-		private SummaryTableColumn CollateColumn(SummaryTableColumn srcColumn, ColumnAggregateType aggregateType, List<int> collatedRowSourceRowCounts, SummaryTableColumn aggregateBaseColumn, out int validMultiRowCollationsCount)
+		private SummaryTableColumn CollateColumn(SummaryTableColumn srcColumn, ColumnAggregateType aggregateType, List<int> collatedRowSourceRowCounts, SummaryTableColumn aggregateBaseColumn, bool bShowMergedStringValues, out int validMultiRowCollationsCount)
 		{
-			SummaryTableColumn collatedColumn = new SummaryTableColumn(srcColumn.name, srcColumn.isNumeric, null, false, srcColumn.elementType, srcColumn.formatInfo, srcColumn.tooltip, srcColumn.columnColourModifier, aggregateType, aggregateBaseColumn);
+			SummaryTableColumn collatedColumn = new SummaryTableColumn(srcColumn.name, srcColumn.isNumeric, null, false, srcColumn.elementType, srcColumn.formatInfo, srcColumn.tooltip, srcColumn.columnColourModifier, aggregateType, aggregateBaseColumn, false, srcColumn.isSortByColumn);
 			collatedColumn.ReserveRows(collatedRowSourceRowCounts.Count);
 
 			validMultiRowCollationsCount = 0;
@@ -1436,31 +1460,37 @@ namespace PerfSummaries
 				}
 				else 
 				{
-					// Collate non-numeric values. Display if they're the same, otherwise show "(multiple)"
-					string mergedValue = null;
-					bool bStringsMatch = true;
+					// Collate non-numeric values
+					HashSet<string> uniqueStringValues = new HashSet<string>();
 					for (int srcRowIndex = srcRowIndexOffset; srcRowIndex < mergedRowEndIndex; srcRowIndex++)
 					{
-						string currentValue = srcColumn.GetStringValue(srcRowIndex);
-						if (mergedValue == null)
+						uniqueStringValues.Add(srcColumn.GetStringValue(srcRowIndex));
+					}
+
+					bool bStringsMatch = true;
+					string mergedValue = "";
+					if (uniqueStringValues.Count > 1)
+					{
+						if (bShowMergedStringValues)
 						{
-							mergedValue = currentValue;
+							// Just show all the unique string values, separated by ;
+							mergedValue = string.Join(';', uniqueStringValues);
 						}
 						else
 						{
-							if (currentValue != mergedValue)
-							{
-								// The strings differ so we can't merge them
-								mergedValue = "<i>(multiple)</i>"; 
-								bStringsMatch = false;
-								break;
-							}
+							// Don't attempt to merge
+							mergedValue = "<i>(multiple)</i>";
 						}
+						bStringsMatch = false;
+					}
+					else if (uniqueStringValues.Count == 1)
+					{
+						mergedValue = uniqueStringValues.First();
 					}
 
 					// If the rows match and are not null or empty then count them.
 					// Note that we do not count as valid merged rows unless Count>1 since otherwise a single count of 1 would cause a row to be visible if CollatedStringVisibility is Auto
-					if (mergedValue != null && mergedValue != "" && bStringsMatch && mergedRowCount > 1)
+					if (mergedValue != "" && bStringsMatch && mergedRowCount > 1)
 					{
 						validMultiRowCollationsCount++;
 					}
@@ -1527,7 +1557,7 @@ namespace PerfSummaries
 			// Generate the collated sortby columns
 			foreach (SummaryTableColumn srcColumn in sortByColumnList)
 			{
-				SummaryTableColumn newColumn = new SummaryTableColumn(srcColumn.name, false, srcColumn.displayName, false, srcColumn.elementType, srcColumn.formatInfo);
+				SummaryTableColumn newColumn = new SummaryTableColumn(srcColumn.name, false, srcColumn.displayName, false, srcColumn.elementType, srcColumn.formatInfo, bInIsSortByColumn:true);
 				newColumns.Add(newColumn);
 				int srcRowIndex = 0;
 				for (int i=0; i<collatedRowSourceRowCounts.Count; i++)
@@ -1549,44 +1579,58 @@ namespace PerfSummaries
 			foreach (SummaryTableColumn column in columns)
 			{
 				// Add avg/min/max columns for this column if it's numeric and we didn't already add it above 
-				if (!collateByColumns.Contains(column))
+				if (collateByColumns.Contains(column))
 				{
-					if (column.isNumeric)
+					continue;
+				}
+				if (column.isNumeric)
+				{
+					if (column.IsDateFormat())
 					{
-						if (column.IsDateFormat())
+						if (dateCollationVisibility != DateCollationVisibility.Hide)
 						{
-							if (dateCollationVisibility != DateCollationVisibility.Hide)
-							{
-								// If we're showing collated dates then determine min/max based on the visibility mode
-								// TODO: expose this to format so we can set per column
-								newColumns.Add(CollateColumn(column, dateCollationVisibility == DateCollationVisibility.Newest ? ColumnAggregateType.Max : ColumnAggregateType.Min, collatedRowSourceRowCounts));
-							}
+							// If we're showing collated dates then determine min/max based on the visibility mode
+							// TODO: expose this to format so we can set per column
+							newColumns.Add(CollateColumn(column, dateCollationVisibility == DateCollationVisibility.Newest ? ColumnAggregateType.Max : ColumnAggregateType.Min, collatedRowSourceRowCounts));
 						}
-						else
+					}
+					else
+					{
+						// Generate aggregate min/max/avg values
+						SummaryTableColumn avgColumn = CollateColumn(column, ColumnAggregateType.Avg, collatedRowSourceRowCounts);
+						newColumns.Add(avgColumn);
+						if (addMinMaxColumns)
 						{
-							// Generate aggregate min/max/avg values
-							SummaryTableColumn avgColumn = CollateColumn(column, ColumnAggregateType.Avg, collatedRowSourceRowCounts);
-							newColumns.Add(avgColumn);
-							if (addMinMaxColumns)
-							{
-								newColumns.Add(CollateColumn(column, ColumnAggregateType.Min, collatedRowSourceRowCounts, avgColumn));
-								newColumns.Add(CollateColumn(column, ColumnAggregateType.Max, collatedRowSourceRowCounts, avgColumn));
-							}
+							newColumns.Add(CollateColumn(column, ColumnAggregateType.Min, collatedRowSourceRowCounts, avgColumn));
+							newColumns.Add(CollateColumn(column, ColumnAggregateType.Max, collatedRowSourceRowCounts, avgColumn));
 						}
 					}
-					else if (stringCollationVisibility != StringCollationVisibility.Hide)
+				}
+				else
+				{
+					if (stringCollationVisibility != StringCollationVisibility.Hide)
 					{
-						SummaryTableColumn stringDataColumn = CollateColumn(column, ColumnAggregateType.None, collatedRowSourceRowCounts, null, out int validMultiRowCollationsCount);
+						SummaryTableColumn stringDataColumn = CollateColumn(column, ColumnAggregateType.None, collatedRowSourceRowCounts, null, false, out int validMultiRowCollationsCount);
 						// If StringCollationVisibility is Auto, only show the column if there were successful valid multi-row collations (ie ignoring collations where row count is 1)
-						if ( validMultiRowCollationsCount > 0 || stringCollationVisibility == StringCollationVisibility.Show )
+						if (validMultiRowCollationsCount > 0 || stringCollationVisibility == StringCollationVisibility.Show)
 						{
 							newColumns.Add(stringDataColumn);
 						}
 					}
+					// Special handling for the CSVID column. Make a hidden CSVIDs column with bShowMergedStringValues=true so we can track the source CSV IDs
+					if (column.elementType == SummaryTableElement.Type.CsvMetadata && column.name == "csvid")
+					{
+						//SummaryTableColumn csvIdsColumn = new SummaryTableColumn("csvids", false, "csvids", false, SummaryTableElement.Type.ToolMetadata, bInIsVisible: false);
+						SummaryTableColumn csvIdsColumn = CollateColumn(column, ColumnAggregateType.None, collatedRowSourceRowCounts, bShowMergedStringValues: true);
+						csvIdsColumn.name = "csvids";
+						csvIdsColumn.isVisible = false;
+						csvIdsColumn.elementType = SummaryTableElement.Type.ToolMetadata;
+						csvIdsColumn.aggregateType = ColumnAggregateType.None;
+						newColumns.Add(csvIdsColumn);
+					}
 				}
 			}
 
-
 			SummaryTable newTable = new SummaryTable();
 			newTable.columns = newColumns;
 			newTable.InitColumnLookup();
@@ -1717,13 +1761,13 @@ namespace PerfSummaries
 			SummaryTable newTable = SortRows(rowSortList, bReverseSort, bSortTrailingDigitsAsNumeric);
 
 			// Generate a column lookup 
-			SummaryTableColumnLookup columnLookup = newTable.GetColumnLookup();
+			SummaryTableColumnLookup newColumnLookup = newTable.GetColumnLookup();
 
 			// Make a list of all unique keys for each element type, including ANY
 			List<string>[] allMetadataKeysLists = new List<string>[(int)SummaryTableElement.Type.COUNT + 1];
 			for (int i = 0; i < allMetadataKeysLists.Length; i++)
 			{
-				allMetadataKeysLists[i] = columnLookup.GetSortedKeyList((SummaryTableElement.Type)i);
+				allMetadataKeysLists[i] = newColumnLookup.GetSortedKeyList((SummaryTableElement.Type)i);
 			}
 
 			// Add columns from the column filter list in the order they appear
@@ -1747,7 +1791,7 @@ namespace PerfSummaries
 					List<SummaryTableColumn> foundColumns = new List<SummaryTableColumn>();
 					foreach (string key in keyList)
 					{
-						SummaryTableColumn column = columnLookup.Get(key, statType);
+						SummaryTableColumn column = newColumnLookup.Get(key, statType);
 						if (column != null)
 						{
 							foundColumns.Add(column);
@@ -1771,14 +1815,16 @@ namespace PerfSummaries
 				}
 				else
 				{
-					SummaryTableColumn column = columnLookup.Get(filterStrLower);
+					SummaryTableColumn column = newColumnLookup.Get(filterStrLower);
 					if (column != null)
 					{
 						newColumnList.Add(column);
 					}
 				}
-
 			}
+
+			SummaryTableColumn csvIdColumn = newColumnLookup.Get("csvid");
+
 			// Remove duplicates
 			newColumnList = newColumnList.Distinct().ToList();
 
@@ -1793,9 +1839,9 @@ namespace PerfSummaries
 			}
 
 			// Compute row weights
-			if (weightByColumnName != null && columnLookup.ContainsKey(weightByColumnName))
+			if (weightByColumnName != null && newColumnLookup.ContainsKey(weightByColumnName))
 			{
-				SummaryTableColumn rowWeightColumn = columnLookup.Get(weightByColumnName);
+				SummaryTableColumn rowWeightColumn = newColumnLookup.Get(weightByColumnName);
 				newTable.rowWeightings = new List<double>(rowWeightColumn.GetCount());
 				for (int i = 0; i < rowWeightColumn.GetCount(); i++)
 				{
@@ -1816,6 +1862,24 @@ namespace PerfSummaries
 			newTable.rowCount = rowCount;
 			newTable.InitColumnLookup();
 
+			// Add a hidden CSV ID column if there isn't a visible one already
+			if (csvIdColumn != null && !newTable.columnLookup.ContainsKey("csvid"))
+			{
+				SummaryTableColumn destCol = csvIdColumn.Clone();
+				destCol.isVisible = false;
+				newTable.columns.Add(destCol);
+				newTable.InitColumnLookup();
+			}
+
+			// Determine the rowSort columns
+			foreach (string columnName in rowSortList )
+			{
+				SummaryTableColumn column = newTable.columnLookup.Get(columnName.ToLower());
+				if (column != null)
+				{
+					column.isSortByColumn = true;
+				}
+			}
 			return newTable;
 		}
 
@@ -2007,20 +2071,21 @@ namespace PerfSummaries
 
 
 		public void WriteToHTML(
-			string htmlFilename, 
-			string VersionString, 
-			bool bSpreadsheetFriendlyStrings, 
-			List<SummarySectionBoundaryInfo> sectionBoundaries, 
-			bool bScrollableTable, 
+			string htmlFilename,
+			string VersionString,
+			bool bSpreadsheetFriendlyStrings,
+			List<SummarySectionBoundaryInfo> sectionBoundaries,
+			bool bScrollableTable,
 			TableColorizeMode tableColorizeMode,
-			bool bAddMinMaxColumns, 
+			bool bAddMinMaxColumns,
 			string hideStatPrefix,
-			int maxColumnStringLength, 
-			string weightByColumnName, 
-			string title, 
+			int maxColumnStringLength,
+			string weightByColumnName,
+			string title,
 			bool bTranspose,
-			bool showFilteredColumnDebug = false,
-			SummaryTableStyle summaryTableStyle = null)
+			SummaryTableStyle summaryTableStyle,
+			bool bEmitJsData,
+			bool showFilteredColumnDebug = false )
 		{
 			summaryTableStyle = summaryTableStyle ?? new SummaryTableStyle();
 
@@ -2074,13 +2139,20 @@ namespace PerfSummaries
 				}
 			}
 
+			//
+			// Output the script block
+			//
+			List<string> scriptLines = new List<string>();
+			if ( bEmitJsData )
+			{
+				scriptLines.Add( GetRowColumnJsData() );
+			}
+
 			if (bScrollableTable)
 			{
 				// Insert some javascript to make the columns sticky. It's not possible to do this for multiple columns with pure CSS, since you need to compute the X offset dynamically
 				// We need to do this when the page is loaded or the window is resized
-				htmlFile.WriteLine("<script>");
-
-				htmlFile.WriteLine(
+				scriptLines.Add(
 					"var originalStyleElement = null; \n" +
 					"document.addEventListener('DOMContentLoaded', function(event) { regenerateStickyColumnCss(); }) \n" +
 					"window.addEventListener('resize', function(event) { regenerateStickyColumnCss(); }) \n" +
@@ -2097,7 +2169,7 @@ namespace PerfSummaries
 				);
 
 				// Make the columns Sticky and compute their X offsets
-				htmlFile.WriteLine(
+				scriptLines.Add(
 					"  var numStickyCols=" + stickyColumnCount + "; \n" +
 					"  var xOffset=0; \n" +
 					"  for (var i=0;i<numStickyCols;i++) \n" +
@@ -2109,9 +2181,21 @@ namespace PerfSummaries
 					"  } \n" +
 					"} \n"
 					);
+			}
+
+			if (scriptLines.Count > 0)
+			{
+				htmlFile.WriteLine("<script>");
+				foreach (string line in scriptLines)
+				{
+					htmlFile.WriteLine(line);
+				}
 				htmlFile.WriteLine("</script>");
 			}
 
+			//
+			// Output CSS
+			//
 			htmlFile.WriteLine("<style type='text/css' id='pageStyle'>");
 			htmlFile.WriteLine("p {  font-family: 'Verdana', Times, serif; font-size: 12px }");
 			htmlFile.WriteLine("h3 {  font-family: 'Verdana', Times, serif; font-size: 14px }");
@@ -2233,7 +2317,11 @@ namespace PerfSummaries
 				// Add the special columns (up to Count) to the lower header row
 				for (int i = 0; i < firstStatColumnIndex; i++)
 				{
-					headerRow.AddCell(columns[i].GetDisplayName(hideStatPrefix, bAddStatNameSpacing, bGreyOutStatPrefixes), String.Join(" ", columns[i].GetHeaderAttributes()));
+					SummaryTableColumn column = columns[i];
+					if (column.isVisible)
+					{
+						headerRow.AddCell(column.GetDisplayName(hideStatPrefix, bAddStatNameSpacing, bGreyOutStatPrefixes), String.Join(" ", column.GetHeaderAttributes()), $"col_{i}");
+					}
 				}
 
 				if (bAddMinMaxColumns)
@@ -2253,6 +2341,11 @@ namespace PerfSummaries
 					for (int i = firstStatColumnIndex; i < columns.Count; i++)
 					{
 						SummaryTableColumn column = columns[i];
+						if (!column.isVisible)
+						{
+							continue;
+						}
+
 						string statName = GetBaseStatNameWithPrefixAndSuffix(column.GetDisplayName(hideStatPrefix, bAddStatNameSpacing, bGreyOutStatPrefixes), out string prefix, out string suffix);
 						string cellText = prefix.Trim();
 						if (column.aggregateType == ColumnAggregateType.None)
@@ -2263,7 +2356,7 @@ namespace PerfSummaries
 						{
 							cellText = column.aggregateType == ColumnAggregateType.Max ? "<i>newest</i>" : "<i>oldest</i>";
 						}
-						headerRow.AddCell(cellText);
+						headerRow.AddCell(cellText, "", $"col_{i}");
 						aggregateColumnCount++;
 
 						// Update the top row if this is the last column of an aggregate (ie max)
@@ -2275,7 +2368,7 @@ namespace PerfSummaries
 								attributes = $"colspan='{aggregateColumnCount}' ";
 							}
 							attributes += String.Join(" ", column.GetHeaderAttributes());
-							topHeaderRow.AddCell(statName + suffix, attributes);
+							topHeaderRow.AddCell(statName + suffix, attributes, $"topCol_{i}");
 							aggregateColumnCount = 0;
 						}
 
@@ -2286,16 +2379,24 @@ namespace PerfSummaries
 					// Add the stat columns
 					for (int i = firstStatColumnIndex; i < columns.Count; i++)
 					{
-						string statName = GetBaseStatNameWithPrefixAndSuffix(columns[i].GetDisplayName(hideStatPrefix, bAddStatNameSpacing, bGreyOutStatPrefixes), out _, out string suffix);
-						headerRow.AddCell(statName + suffix, String.Join(" ", columns[i].GetHeaderAttributes()));
+						SummaryTableColumn column = columns[i];
+						if (column.isVisible)
+						{
+							string statName = GetBaseStatNameWithPrefixAndSuffix(column.GetDisplayName(hideStatPrefix, bAddStatNameSpacing, bGreyOutStatPrefixes), out _, out string suffix);
+							headerRow.AddCell(statName + suffix, String.Join(" ", column.GetHeaderAttributes()));
+						}
 					}
 				}
 			}
 			else
 			{
-				foreach (SummaryTableColumn column in columns)
+				for (int i = 0; i < columns.Count; i++)
 				{
-					headerRow.AddCell(column.GetDisplayName(hideStatPrefix, bAddStatNameSpacing, bGreyOutStatPrefixes), String.Join(" ", column.GetHeaderAttributes()));
+					SummaryTableColumn column = columns[i];
+					if (column.isVisible)
+					{
+						headerRow.AddCell(column.GetDisplayName(hideStatPrefix, bAddStatNameSpacing, bGreyOutStatPrefixes), String.Join(" ", column.GetHeaderAttributes()), $"col_{i}");
+					}
 				}
 			}
 			htmlTable.AddRow(headerRow);
@@ -2393,6 +2494,10 @@ namespace PerfSummaries
 				int columnIndex = 0;
 				foreach (SummaryTableColumn column in columns)
 				{
+					if (!column.isVisible)
+					{
+						continue;
+					}
 					List<string> attributes = new List<string>();
 
 					// Add the tooltip for non-collated tables
@@ -2512,6 +2617,20 @@ namespace PerfSummaries
 				}
 			}
 
+			// Assign an ID to rows containing data
+			for (int i = 0; i < htmlTable.rows.Count; i++)
+			{
+				HtmlTable.Row row = htmlTable.rows[i];
+				if (i < htmlTable.numHeaderRows)
+				{
+					row.id = $"headerRow_{i}";
+				}
+				else
+				{
+					row.id = $"row_{i - htmlTable.numHeaderRows}";
+				}
+			}
+
 			htmlTable.WriteToHtml(htmlFile);
 
 			string extraString = "";
@@ -2528,6 +2647,7 @@ namespace PerfSummaries
 
 		public SummaryTable SortRows(List<string> rowSortList, bool reverseSort, bool sortTrailingDigitsAsNumeric = false)
 		{
+			// Generate a sort key for each row
 			List<KeyValuePair<string, int>> columnRemapping = new List<KeyValuePair<string, int>>();
 			for (int i = 0; i < rowCount; i++)
 			{
@@ -2567,12 +2687,13 @@ namespace PerfSummaries
 				columnRemapping.Add(new KeyValuePair<string, int>(key, i));
 			}
 
+			// Sort the columns
 			columnRemapping.Sort(delegate (KeyValuePair<string, int> m1, KeyValuePair<string, int> m2)
 			{
 				return m1.Key.CompareTo(m2.Key);
 			});
 
-			// Reorder the columns
+			// Reorder the data in each column
 			List<SummaryTableColumn> newColumns = new List<SummaryTableColumn>();
 			foreach (SummaryTableColumn srcCol in columns)
 			{
@@ -2597,12 +2718,60 @@ namespace PerfSummaries
 			SummaryTable newTable = new SummaryTable();
 			newTable.columns = newColumns;
 			newTable.rowCount = rowCount;
-			newTable.firstStatColumnIndex = firstStatColumnIndex;
+			newTable.firstStatColumnIndex = firstStatColumnIndex; 
 			newTable.isCollated = isCollated;
 			newTable.InitColumnLookup();
 			return newTable;
 		}
 
+		// Writes a couple of javascript objects containing row and column info
+		string GetRowColumnJsData()
+		{
+			List<string> lines = new List<string>();
+			List<string> RowEntries = new List<string>();
+			for (int rowIndex = 0; rowIndex < rowCount; rowIndex++)
+			{
+				if (!IsRowVisible(rowIndex))
+				{
+					continue;
+				}
+				// Output a dict for the row containing just the entries referenced in rowSort
+				Dictionary<string, dynamic> rowDict = new Dictionary<string, dynamic>();
+
+				// Add the csvID row data. Use the csvids column (hidden collated column containing all source csvids) if it exists, otherwise csvid
+				SummaryTableColumn csvIdsColumn = columnLookup.Get("csvids", SummaryTableElement.Type.ToolMetadata);
+				if (csvIdsColumn != null)
+				{
+					rowDict.Add(csvIdsColumn.name, csvIdsColumn.GetStringValue(rowIndex).Split(';'));
+				}
+				else
+				{
+					SummaryTableColumn csvIdColumn = columnLookup.Get("csvid", SummaryTableElement.Type.CsvMetadata);
+					rowDict.Add(csvIdColumn.name, csvIdColumn.GetStringValue(rowIndex));
+				}
+
+				foreach (SummaryTableColumn column in columns)
+				{ 
+					if ( column.isSortByColumn && !rowDict.ContainsKey(column.name))
+					{
+						string value = column.formatInfo == null ? column.GetStringValue(rowIndex) : column.GetStringValue(rowIndex, false, column.formatInfo.numericFormat);
+						rowDict.Add(column.name, value);
+					}
+				}
+				RowEntries.Add(JsonSerializer.Serialize(rowDict));
+			}
+			lines.Add($"SummaryTableRowInfo=[{String.Join(", ", RowEntries)}];\n");
+
+			List<string> ColumnEntries = new List<string>();
+			foreach (SummaryTableColumn column in columns)
+			{
+				ColumnEntries.Add(column.GetJsData());
+			}
+			lines.Add($"SummaryTableColumnInfo=[{String.Join(", ", ColumnEntries)}];\n");
+
+			return String.Join("\n", lines);
+		}
+
 		void InitColumnLookup()
 		{
 			columnLookup.Clear();
@@ -2713,9 +2882,10 @@ namespace PerfSummaries
 	{
 		public class Row
 		{
-			public Row(string inAttributes="", int reserveColumnCount=0)
+			public Row(string inAttributes="", int reserveColumnCount=0, string inId=null)
 			{
 				attributes = inAttributes;
+				id = inId;
 				if (reserveColumnCount > 0)
 				{
 					cells = new List<Cell>(reserveColumnCount);
@@ -2726,24 +2896,25 @@ namespace PerfSummaries
 				}
 			}
 
-			public void AddCell(string contents="", string attributes="")
+			public void AddCell(string contents="", string attributes="", string cellId = null)
 			{
-				cells.Add(new Cell(contents, attributes));
+				cells.Add(new Cell(contents, attributes, false, cellId));
 			}
 
-			public void AddCell(string contents, List<string> attributes, Dictionary<string, string> styleAttributes)
+			public void AddCell(string contents, List<string> attributes, Dictionary<string, string> styleAttributes, string cellId = null)
 			{
 				List<string> styleLines = styleAttributes.Select(pair => $"{pair.Key}:{pair.Value}").ToList();
 				attributes.Add($"style='{String.Join(";", styleLines)}'");
 
-				cells.Add(new Cell(contents, String.Join(" ",attributes)));
+				cells.Add(new Cell(contents, String.Join(" ",attributes), false, cellId));
 			}
 
 
 			public void WriteToHtml(System.IO.StreamWriter htmlFile, bool bIsHeaderRow )
 			{
 				string attributesStr = attributes.Length == 0 ? "" : " " + attributes;
-				htmlFile.Write("<tr"+ attributesStr + ">");
+				string idStr = id == null ? "" : $" id=\"{id}\"";
+				htmlFile.Write("<tr" + idStr + attributesStr + ">");
 				foreach (Cell cell in cells)
 				{
 					if (cell == null)
@@ -2772,6 +2943,7 @@ namespace PerfSummaries
 			}
 
 			public string attributes;
+			public string id;
 			public List<Cell> cells;
 		}
 
@@ -2784,23 +2956,27 @@ namespace PerfSummaries
 			{
 				contents = srcCell.contents;
 				attributes = srcCell.attributes;
+				id = srcCell.id;
 			}
 
-			public Cell(string inContents, string inAttributes, bool bInIsHeader=false)
+			public Cell(string inContents, string inAttributes, bool bInIsHeader=false, string inId = null)
 			{
 				contents = inContents;
 				attributes = inAttributes;
+				id = inId;
 			}
 
 			public void WriteToHtml(System.IO.StreamWriter htmlFile, bool bIsHeaderRow)
 			{
 				string AttributesStr = attributes.Length == 0 ? "" : " "+ attributes;
+				string IdStr = id == null ? "" : $" id=\"{id}\"";
 				string CellType = bIsHeaderRow ? "th" : "td";
-				htmlFile.Write("<" + CellType + AttributesStr + ">" + contents + "</" + CellType + ">");
+				htmlFile.Write("<" + CellType + IdStr + AttributesStr+">" + contents + "</" + CellType + ">");
 			}
 
 			public string contents;
 			public string attributes;
+			public string id;
 		}
 
 		public HtmlTable(string inAttributes="", int inNumHeaderRows=1, int reserveColumnCount=0)
