PerfReportTool 4.252.0. Simplify summary table column ID assignment and fix bug with transposed summary tables where column IDs were assigned per row instead of per column because they were set  before transposing
Add SummaryTableIsTransposed bool to the script block so we can detect transposed summary tables downstream.
Note: SummaryTableRowData and SummaryTableColumnData in the script block still refers to pre-transposed values, e.g so column data refers to rows if the table is transposed.

diff --git a/Engine/Source/Programs/CSVTools/PerfReportTool/PerfReportTool.cs b/Engine/Source/Programs/CSVTools/PerfReportTool/PerfReportTool.cs
--- a/Engine/Source/Programs/CSVTools/PerfReportTool/PerfReportTool.cs
+++ b/Engine/Source/Programs/CSVTools/PerfReportTool/PerfReportTool.cs
@@ -23,7 +23,7 @@ namespace PerfReportTool
     class Version
     {
 		// Format: Major.Minor.Bugfix
-        private static string VersionString = "4.251.0";
+        private static string VersionString = "4.252.0";
 
         public static string Get() { return VersionString; }
     };
diff --git a/Engine/Source/Programs/CSVTools/PerfReportTool/SummaryTable.cs b/Engine/Source/Programs/CSVTools/PerfReportTool/SummaryTable.cs
--- a/Engine/Source/Programs/CSVTools/PerfReportTool/SummaryTable.cs
+++ b/Engine/Source/Programs/CSVTools/PerfReportTool/SummaryTable.cs
@@ -2147,6 +2147,8 @@ namespace PerfSummaries
 			{
 				scriptLines.Add( GetRowColumnJsData() );
 			}
+			scriptLines.Add($"SummaryTableIsTransposed={(bTranspose ? "true" : "false")};");
+
 
 			if (bScrollableTable)
 			{
@@ -2320,7 +2322,7 @@ namespace PerfSummaries
 					SummaryTableColumn column = columns[i];
 					if (column.isVisible)
 					{
-						headerRow.AddCell(column.GetDisplayName(hideStatPrefix, bAddStatNameSpacing, bGreyOutStatPrefixes), String.Join(" ", column.GetHeaderAttributes()), $"col_{i}");
+						headerRow.AddCell(column.GetDisplayName(hideStatPrefix, bAddStatNameSpacing, bGreyOutStatPrefixes), String.Join(" ", column.GetHeaderAttributes()));
 					}
 				}
 
@@ -2356,7 +2358,7 @@ namespace PerfSummaries
 						{
 							cellText = column.aggregateType == ColumnAggregateType.Max ? "<i>newest</i>" : "<i>oldest</i>";
 						}
-						headerRow.AddCell(cellText, "", $"col_{i}");
+						headerRow.AddCell(cellText, "");
 						aggregateColumnCount++;
 
 						// Update the top row if this is the last column of an aggregate (ie max)
@@ -2368,7 +2370,7 @@ namespace PerfSummaries
 								attributes = $"colspan='{aggregateColumnCount}' ";
 							}
 							attributes += String.Join(" ", column.GetHeaderAttributes());
-							topHeaderRow.AddCell(statName + suffix, attributes, $"topCol_{i}");
+							topHeaderRow.AddCell(statName + suffix, attributes);
 							aggregateColumnCount = 0;
 						}
 
@@ -2395,7 +2397,7 @@ namespace PerfSummaries
 					SummaryTableColumn column = columns[i];
 					if (column.isVisible)
 					{
-						headerRow.AddCell(column.GetDisplayName(hideStatPrefix, bAddStatNameSpacing, bGreyOutStatPrefixes), String.Join(" ", column.GetHeaderAttributes()), $"col_{i}");
+						headerRow.AddCell(column.GetDisplayName(hideStatPrefix, bAddStatNameSpacing, bGreyOutStatPrefixes), String.Join(" ", column.GetHeaderAttributes()));
 					}
 				}
 			}
@@ -2587,6 +2589,7 @@ namespace PerfSummaries
 				}
 			}
 
+			// Transpose the table if requested
 			if (bTranspose)
 			{
 				htmlTable = htmlTable.Transpose();
@@ -2596,6 +2599,24 @@ namespace PerfSummaries
 				htmlTable.rows[htmlTable.numHeaderRows + firstStatColumnIndex - 1].attributes = " class='sectionStartLevel2'";
 			}
 
+			// Assign IDs to rows and columns
+			for (int i = 0; i < htmlTable.rows.Count; i++)
+			{
+				HtmlTable.Row row = htmlTable.rows[i];
+				if (i < htmlTable.numHeaderRows)
+				{
+					row.id = $"headerRow_{i}";
+					for (int j=0; j<row.cells.Count; j++)
+					{
+						row.cells[j].id = (i == htmlTable.numHeaderRows - 1) ? $"col_{j}" : $"topCol_{j}";
+					}
+				}
+				else
+				{
+					row.id = $"row_{i - htmlTable.numHeaderRows}";
+				}
+			}
+
 			// Apply final formatting
 			htmlTable.rows[htmlTable.numHeaderRows-1].attributes += " class='lastHeaderRow'";
 
@@ -2617,20 +2638,6 @@ namespace PerfSummaries
 				}
 			}
 
-			// Assign an ID to rows containing data
-			for (int i = 0; i < htmlTable.rows.Count; i++)
-			{
-				HtmlTable.Row row = htmlTable.rows[i];
-				if (i < htmlTable.numHeaderRows)
-				{
-					row.id = $"headerRow_{i}";
-				}
-				else
-				{
-					row.id = $"row_{i - htmlTable.numHeaderRows}";
-				}
-			}
-
 			htmlTable.WriteToHtml(htmlFile);
 
 			string extraString = "";
