PerfReportTool 4.253.0: Derived stat reporting improvements
- Silence warnings if source stats referenced in derived stat expressions are missing. These should be treated the same as unreferenced stats generally. We can pass -derivedStatVerboseErrors to override
- Make the commandline of CommandLineApp publically accessible to avoid having to pass params around if we don't need to

diff --git a/Engine/Source/Programs/CSVTools/CsvStats/CommandLineTool.cs b/Engine/Source/Programs/CSVTools/CsvStats/CommandLineTool.cs
--- a/Engine/Source/Programs/CSVTools/CsvStats/CommandLineTool.cs
+++ b/Engine/Source/Programs/CSVTools/CsvStats/CommandLineTool.cs
@@ -230,7 +230,7 @@ namespace CSVStats
 			Linux
 		}
 
-		protected CommandLine commandLine;
+		public static CommandLine commandLine;
 
 		private readonly static bool bIsMac = File.Exists("/System/Library/CoreServices/SystemVersion.plist");
 
diff --git a/Engine/Source/Programs/CSVTools/CsvStats/CsvStats.cs b/Engine/Source/Programs/CSVTools/CsvStats/CsvStats.cs
--- a/Engine/Source/Programs/CSVTools/CsvStats/CsvStats.cs
+++ b/Engine/Source/Programs/CSVTools/CsvStats/CsvStats.cs
@@ -72,7 +72,18 @@ namespace CSVStats
         public int Frame;
     };
 
-    public class CsvMetadata
+	public class DerivedStatException : Exception
+	{
+		public DerivedStatException(string message) : base(message) { }
+	}
+
+	public class DerivedStatMissingSourceException : DerivedStatException
+	{
+		public DerivedStatMissingSourceException(string message) : base(message) { }
+	}
+
+
+	public class CsvMetadata
     {
         public Dictionary<string, string> Values;
 
@@ -609,10 +620,13 @@ namespace CSVStats
 						string key = ReadQuotedString(paramString);
 						key = (key == null) ? paramString : key;
 						string val = GetMetadataValue(key);
-						float valueNumeric;
-						if (val==null || !float.TryParse(val, out valueNumeric))
+						if (val == null)
+						{
+							throw new DerivedStatMissingSourceException("Metadata not found: " + expression[oldIndex..index]);
+						}
+						if (!float.TryParse(val, out _))
 						{
-							throw new Exception("Warning: failed to read metadata: " + expression[oldIndex..index]);
+							throw new DerivedStatException("Failed to parse metadata : " + expression[oldIndex..index]);
 						}
 						return val.ToString();
 					}
@@ -623,7 +637,7 @@ namespace CSVStats
 						eventStr = (eventStr == null) ? paramString : eventStr;
 						if (eventStr == null)
 						{
-							throw new Exception("Warning: eventstr requires an arg: " + expression[oldIndex..index]);
+							throw new DerivedStatException("Eventstr requires an arg: " + expression[oldIndex..index]);
 						}
 						string statName = "__STAT__EVENTCOUNT__" + tokenIndex.ToString();
 						statsOut[statName] = MakeEventCountStat(statName, eventStr);
@@ -658,13 +672,13 @@ namespace CSVStats
 						srcStatName = (srcStatName == null) ? paramString : srcStatName;
 						if (srcStatName == null)
 						{
-							throw new Exception("Warning: frameDelta requires an arg: " + expression[oldIndex..index]);
+							throw new DerivedStatException("FrameDelta requires an arg: " + expression[oldIndex..index]);
 						}
 
 						StatSamples srcStat = GetStat(srcStatName);
 						if (srcStat == null)
 						{
-							throw new Exception("Warning: source stat not found: " + expression[oldIndex..index]);
+							throw new DerivedStatMissingSourceException("Source stat not found: " + expression[oldIndex..index]);
 						}
 						string destStatName = "__STAT__FRAMEDELTA__" + tokenIndex.ToString();
 						StatSamples newStat = new StatSamples(destStatName, srcStat.samples.Count);
@@ -687,7 +701,7 @@ namespace CSVStats
 				StatSamples stat = GetStat(statName.Trim());
 				if (stat == null)
 				{
-					throw new Exception("Stat not found: " + statName);
+					throw new DerivedStatMissingSourceException("Stat not found: " + statName);
 				}
 				string safeName = stat.Name.Replace("/", "__SLASH__");
 				statsOut[safeName] = stat;
@@ -738,7 +752,7 @@ namespace CSVStats
 			return String.Join("", tokens);
 		}
 
-		public StatSamples MakeDerivedStatFromExpression(string expression)
+		public StatSamples MakeDerivedStatFromExpression(string expression, bool bVerboseErrors = false)
 		{
 			if (expression == "")
 			{
@@ -781,9 +795,12 @@ namespace CSVStats
 			{
 				finalExpression = PreProcessExpression(expression, referencedStats);
 			}
-			catch (Exception ex)
+			catch (DerivedStatException ex)
 			{
-				Console.WriteLine("Error: derived stat expression \"" + expression + "\"could not be computed. Preprocessing error:" + ex.ToString());
+				if (!(ex is DerivedStatMissingSourceException) || bVerboseErrors)
+				{
+					Console.WriteLine("Error: derived stat expression \"" + expression + "\"could not be computed. Preprocessing error:" + ex.ToString());
+				}
 				return null;
 			}
 
@@ -830,9 +847,9 @@ namespace CSVStats
 			return statSamples;
 		}
 
-		public StatSamples AddDerivedStatFromExpression(string expression)
+		public StatSamples AddDerivedStatFromExpression(string expression, bool bVerboseErrors = false)
 		{
-			StatSamples derivedStat = MakeDerivedStatFromExpression(expression);
+			StatSamples derivedStat = MakeDerivedStatFromExpression(expression, bVerboseErrors);
 			if (derivedStat != null)
 			{
 				AddStat(derivedStat);
diff --git a/Engine/Source/Programs/CSVTools/CsvToSvgLib/CsvToSvgLib.cs b/Engine/Source/Programs/CSVTools/CsvToSvgLib/CsvToSvgLib.cs
--- a/Engine/Source/Programs/CSVTools/CsvToSvgLib/CsvToSvgLib.cs
+++ b/Engine/Source/Programs/CSVTools/CsvToSvgLib/CsvToSvgLib.cs
@@ -1417,6 +1417,8 @@ namespace CSVTools
 			public int Count = 0;
 		};
 
+		// TODO: this would be better if merged the best points first, ie the ones where the error value is close to zero
+		// Also, the approach of using only the current pass points error to compound. Ideally we want it to be computed from the original points
 		Vec2[] RemoveRedundantPoints(Vec2[] RawPoints, float compressionThreshold, int passIndex)
 		{
 			//return RawPoints;
diff --git a/Engine/Source/Programs/CSVTools/PerfReportTool/CsvFileCache..cs b/Engine/Source/Programs/CSVTools/PerfReportTool/CsvFileCache..cs
--- a/Engine/Source/Programs/CSVTools/PerfReportTool/CsvFileCache..cs
+++ b/Engine/Source/Programs/CSVTools/PerfReportTool/CsvFileCache..cs
@@ -121,9 +121,10 @@ namespace PerfReportTool
 
 		private void AddDerivedMetadataToCsv(CsvStats csvStats)
 		{
+			bool bVerboseErrors = CommandLineTool.commandLine.GetBoolArg("derivedStatVerboseErrors", false);
 			foreach (DerivedCsvStatDefinition derivedStat in derivedCsvStatDefinitions)
 			{
-				finalCsv.AddDerivedStatFromExpression(derivedStat.getExpressionStr());
+				finalCsv.AddDerivedStatFromExpression(derivedStat.getExpressionStr(), bVerboseErrors);
 			}
 		}
 
diff --git a/Engine/Source/Programs/CSVTools/PerfReportTool/PerfReportTool.cs b/Engine/Source/Programs/CSVTools/PerfReportTool/PerfReportTool.cs
--- a/Engine/Source/Programs/CSVTools/PerfReportTool/PerfReportTool.cs
+++ b/Engine/Source/Programs/CSVTools/PerfReportTool/PerfReportTool.cs
@@ -23,7 +23,7 @@ namespace PerfReportTool
     class Version
     {
 		// Format: Major.Minor.Bugfix
-        private static string VersionString = "4.252.0";
+        private static string VersionString = "4.253.0";
 
         public static string Get() { return VersionString; }
     };
@@ -181,6 +181,7 @@ namespace PerfReportTool
 			"  -collatedDateVisibility <hide|newest|oldest> : sets visibility of dates in collated summary tables (default=newest)\n" +
 			"  -columnSortMode <default|sortByMax|sortByAvg|wildCardSortByMax|wildCardSortByAvg> : sorting mode for columns\n" +
 			"  -summaryTableJsData : embeds javascript row and column data in the <script> section to help with interop\n" +
+            "  -derivedStatVerboseErrors : output verbose errors for derived stat expressions\n" +
 
 			"Diff rows\n" +
 			"  -addDiffRows : adds diff rows after the first two rows\n" +
