PRT - Don't truncate link templates in summary tables

#rnx

diff --git a/Engine/Source/Programs/CSVTools/PerfReportTool/PerfReportTool.cs b/Engine/Source/Programs/CSVTools/PerfReportTool/PerfReportTool.cs
--- a/Engine/Source/Programs/CSVTools/PerfReportTool/PerfReportTool.cs
+++ b/Engine/Source/Programs/CSVTools/PerfReportTool/PerfReportTool.cs
@@ -23,7 +23,7 @@ namespace PerfReportTool
     class Version
     {
 		// Format: Major.Minor.Bugfix
-        private static string VersionString = "4.253.0";
+        private static string VersionString = "4.253.1";
 
         public static string Get() { return VersionString; }
     };
diff --git a/Engine/Source/Programs/CSVTools/PerfReportTool/SummaryTable.cs b/Engine/Source/Programs/CSVTools/PerfReportTool/SummaryTable.cs
--- a/Engine/Source/Programs/CSVTools/PerfReportTool/SummaryTable.cs
+++ b/Engine/Source/Programs/CSVTools/PerfReportTool/SummaryTable.cs
@@ -43,6 +43,11 @@ namespace PerfSummaries
 					return anchor + text + "</a>";
 				}
 			}
+			else if (inValue.StartsWith("{LinkTemplate:"))
+			{
+				// Don't truncate link templates.
+				return inValue;
+			}
 			return SanitizeHtmlString(inValue.Substring(0, maxLength)) + "...";
 		}
 	}
