Update the help text in CsvConvert to better reflect usage.

The help text for outFormat and binCompress previously showed the form '-arg=<option1|option2>' but the expects '-arg option'. Including the '=' in the argument results in failure to parse the parameter. This change removes '=' in the help text to avoid confusion.

diff --git a/Engine/Source/Programs/CSVTools/CsvConvert/CsvConvert.cs b/Engine/Source/Programs/CSVTools/CsvConvert/CsvConvert.cs
--- a/Engine/Source/Programs/CSVTools/CsvConvert/CsvConvert.cs
+++ b/Engine/Source/Programs/CSVTools/CsvConvert/CsvConvert.cs
@@ -59,8 +59,8 @@ namespace CSVInfo
 				"  <filename>\n" +
 				"OR \n" +
 				"  -in <filename>\n" +
-				"  -outFormat=<csv|bin|csvNoMetadata>\n" +
-				"  [-binCompress=<0|1|2>] (0=none)\n" +
+				"  -outFormat <csv|bin|csvNoMetadata>\n" +
+				"  [-binCompress <0|1|2>] (0=none)\n" +
 				"  [-o outFilename]\n" +
 				"  [-force] - overwrite even if output file exists already\n" +
 				"  [-inPlace] - overwrite the input file (implies -force)\n" +
