PerfReportTool 4.253.2: Fix an edge case exception when transposing summary tables with no stats

#DisabledValidations Preflight, CheckCopyright

diff --git a/Engine/Source/Programs/CSVTools/PerfReportTool/PerfReportTool.cs b/Engine/Source/Programs/CSVTools/PerfReportTool/PerfReportTool.cs
--- a/Engine/Source/Programs/CSVTools/PerfReportTool/PerfReportTool.cs
+++ b/Engine/Source/Programs/CSVTools/PerfReportTool/PerfReportTool.cs
@@ -23,7 +23,7 @@ namespace PerfReportTool
     class Version
     {
 		// Format: Major.Minor.Bugfix
-        private static string VersionString = "4.253.1";
+        private static string VersionString = "4.253.2";
 
         public static string Get() { return VersionString; }
     };
diff --git a/Engine/Source/Programs/CSVTools/PerfReportTool/SummaryTable.cs b/Engine/Source/Programs/CSVTools/PerfReportTool/SummaryTable.cs
--- a/Engine/Source/Programs/CSVTools/PerfReportTool/SummaryTable.cs
+++ b/Engine/Source/Programs/CSVTools/PerfReportTool/SummaryTable.cs
@@ -2601,7 +2601,11 @@ namespace PerfSummaries
 				htmlTable.Set(0, 0, new HtmlTable.Cell(title, ""));
 
 				// Add a section boundary where the stats start
-				htmlTable.rows[htmlTable.numHeaderRows + firstStatColumnIndex - 1].attributes = " class='sectionStartLevel2'";
+				int firstStatRowIndex = htmlTable.numHeaderRows + firstStatColumnIndex - 1;
+				if (firstStatRowIndex < htmlTable.rows.Count)
+				{
+					htmlTable.rows[firstStatRowIndex].attributes = " class='sectionStartLevel2'";
+				}
 			}
 
 			// Assign IDs to rows and columns
