PerfReportTool 4.253.3: Fix an edge case bug when adding diff rows which causes all columns to be culled by -columnDiffDisplayThreshold. Skip adding diff rows if we have less than 2 rows

diff --git a/Engine/Source/Programs/CSVTools/PerfReportTool/PerfReportTool.cs b/Engine/Source/Programs/CSVTools/PerfReportTool/PerfReportTool.cs
--- a/Engine/Source/Programs/CSVTools/PerfReportTool/PerfReportTool.cs
+++ b/Engine/Source/Programs/CSVTools/PerfReportTool/PerfReportTool.cs
@@ -23,7 +23,7 @@ namespace PerfReportTool
     class Version
     {
 		// Format: Major.Minor.Bugfix
-        private static string VersionString = "4.253.2";
+        private static string VersionString = "4.253.3";
 
         public static string Get() { return VersionString; }
     };
diff --git a/Engine/Source/Programs/CSVTools/PerfReportTool/SummaryTable.cs b/Engine/Source/Programs/CSVTools/PerfReportTool/SummaryTable.cs
--- a/Engine/Source/Programs/CSVTools/PerfReportTool/SummaryTable.cs
+++ b/Engine/Source/Programs/CSVTools/PerfReportTool/SummaryTable.cs
@@ -1890,6 +1890,11 @@ namespace PerfSummaries
 
 		public void AddDiffRows(bool bSortColumnsByDiff, double columnDiffDisplayThreshold, bool bInShowOnlyDiffRows, bool bDiffRowsAlternating)
 		{
+			// We can't add diff rows unless we have at least 2 rows
+			if (rowCount < 2)
+			{
+				return;
+			}
 			diffRowFrequency = bDiffRowsAlternating ? DiffRowFrequency.Alternating : DiffRowFrequency.AfterEachPair;
 			bShowOnlyDiffRows = bInShowOnlyDiffRows;
 			for (int i=0; i<columns.Count; i++)
