PerfReportTool 4.253.4: Pre-size divs for dynamically loaded graphs so loading doesn't affect the layout

diff --git a/Engine/Source/Programs/CSVTools/PerfReportTool/PerfReportTool.cs b/Engine/Source/Programs/CSVTools/PerfReportTool/PerfReportTool.cs
--- a/Engine/Source/Programs/CSVTools/PerfReportTool/PerfReportTool.cs
+++ b/Engine/Source/Programs/CSVTools/PerfReportTool/PerfReportTool.cs
@@ -23,7 +23,7 @@ namespace PerfReportTool
     class Version
     {
 		// Format: Major.Minor.Bugfix
-        private static string VersionString = "4.253.3";
+        private static string VersionString = "4.253.4";
 
         public static string Get() { return VersionString; }
     };
@@ -1713,7 +1713,7 @@ namespace PerfReportTool
 
 						string graphUrlRoot = GetArg("embeddedGraphUrlRoot", mandatory: true);
 						string graphUrl = $"{graphUrlRoot}?csvs={csvId}&{graphArgs}";
-						string script = $"<div id=\"graph_{svgFileIndex}\"></div>\n";
+						string script = $"<div id=\"graph_{svgFileIndex}\" style=\"width:{graph.settings.width.value}px;height:{graph.settings.height.value}px\"></div>\n";
 						script += $"<script>" +
 							// On page load, fetch the graph from the end point
 							$"fetch('{graphUrl}')\n" +
