CsvToSvg 3.69/CsvTools: fix bug with multithreaded commandline generation in CsvToSVG. Batch processing creates multiple program instances, so we need commandline to remain non-static.
- To access the commandline statically, you can now call CommandLineTool.GetGlobalCommandLine() if necessary. This won't work for batch instances
- Make it so CsvToSvg doesn't catch exceptions if a debugger is attached (unless "-catchExceptions" is passed)

#tests PRT test cases

diff --git a/Engine/Source/Programs/CSVTools/CSVToSVG/CsvToSvgProgram.cs b/Engine/Source/Programs/CSVTools/CSVToSVG/CsvToSvgProgram.cs
--- a/Engine/Source/Programs/CSVTools/CSVToSVG/CsvToSvgProgram.cs
+++ b/Engine/Source/Programs/CSVTools/CSVToSVG/CsvToSvgProgram.cs
@@ -264,25 +264,30 @@ namespace CSVTools
 			GraphGenerator generator = new GraphGenerator(csvs);
 
 			bool writeErrorsToSVG = GetIntArg("writeErrorsToSVG", 1) == 1;
-			try
+			if (Debugger.IsAttached && !GetBoolArg("catchExceptions"))
 			{
 				generator.MakeGraph(graphParams, svgFilename, writeErrorsToSVG);
-
-				if ( GetBoolArg("view") )
+			}
+			else
+			{
+				try
 				{
-					ViewGraph(System.IO.Path.GetFullPath(svgFilename));
+					generator.MakeGraph(graphParams, svgFilename, writeErrorsToSVG);
+				}
+				catch (System.Exception e)
+				{
+					// Write the error to the SVG
+					string errorString = e.ToString();
+					Console.Out.WriteLine("Error: " + e.ToString());
 				}
-
 			}
-			catch (System.Exception e)
+			if (GetBoolArg("view"))
 			{
-				// Write the error to the SVG
-				string errorString = e.ToString();
-				Console.Out.WriteLine("Error: " + e.ToString());
+				ViewGraph(System.IO.Path.GetFullPath(svgFilename));
 			}
 		}
 
-        static string[] MakeArgsArray(string argsStr)
+		static string[] MakeArgsArray(string argsStr)
         {
             List<string> argsOut = new List<string>();
             string currentArg = "";
diff --git a/Engine/Source/Programs/CSVTools/CsvStats/CommandLineTool.cs b/Engine/Source/Programs/CSVTools/CsvStats/CommandLineTool.cs
--- a/Engine/Source/Programs/CSVTools/CsvStats/CommandLineTool.cs
+++ b/Engine/Source/Programs/CSVTools/CsvStats/CommandLineTool.cs
@@ -6,6 +6,7 @@ using System.Collections.Generic;
 using System.Text;
 using System.IO;
 using System.Diagnostics;
+using System.Threading;
 
 namespace CSVStats
 {
@@ -230,7 +231,28 @@ namespace CSVStats
 			Linux
 		}
 
-		public static CommandLine commandLine;
+		static int initCounter = 0;
+
+		public CommandLineTool()
+		{
+			bMainInstance = Interlocked.Increment(ref initCounter) == 1;
+		}
+
+		// Accesses the commandline for the main instance. Caution: if using batch CommandLineTool instances (e.g via CsvToSvg -batchCommands), this will not return the commandline for the batch
+		public static CommandLine GetGlobalCommandLine()
+		{
+			return globalCommandLine;
+		}
+		public CommandLine GetCommandLine()
+		{
+			return commandLine;
+		}
+
+		// It's possible to create instances of commandlineTool for batch processing. When this happens, bMainInstance will be false
+		// globalCommandline is only valid for the main/first instance
+		protected CommandLine commandLine;
+		private static CommandLine globalCommandLine;
+		private bool bMainInstance = true;
 
 		private readonly static bool bIsMac = File.Exists("/System/Library/CoreServices/SystemVersion.plist");
 
@@ -369,6 +391,10 @@ namespace CSVStats
 		protected void ReadCommandLine(string[] args)
 		{
 			commandLine = new CommandLine(args);
+			if (bMainInstance)
+			{
+				globalCommandLine = commandLine;
+			}
 		}
 	}
 
diff --git a/Engine/Source/Programs/CSVTools/CsvToSvgLib/CsvToSvgLib.cs b/Engine/Source/Programs/CSVTools/CsvToSvgLib/CsvToSvgLib.cs
--- a/Engine/Source/Programs/CSVTools/CsvToSvgLib/CsvToSvgLib.cs
+++ b/Engine/Source/Programs/CSVTools/CsvToSvgLib/CsvToSvgLib.cs
@@ -15,7 +15,7 @@ namespace CSVTools
 
 	public class CsvToSvgLibVersion
 	{
-		private static string VersionString = "3.68";
+		private static string VersionString = "3.69";
 
 		public static string Get() { return VersionString; }
 	};
diff --git a/Engine/Source/Programs/CSVTools/PerfReportTool/CsvFileCache..cs b/Engine/Source/Programs/CSVTools/PerfReportTool/CsvFileCache..cs
--- a/Engine/Source/Programs/CSVTools/PerfReportTool/CsvFileCache..cs
+++ b/Engine/Source/Programs/CSVTools/PerfReportTool/CsvFileCache..cs
@@ -121,7 +121,7 @@ namespace PerfReportTool
 
 		private void AddDerivedMetadataToCsv(CsvStats csvStats)
 		{
-			bool bVerboseErrors = CommandLineTool.commandLine.GetBoolArg("derivedStatVerboseErrors", false);
+			bool bVerboseErrors = CommandLineTool.GetGlobalCommandLine().GetBoolArg("derivedStatVerboseErrors", false);
 			foreach (DerivedCsvStatDefinition derivedStat in derivedCsvStatDefinitions)
 			{
 				finalCsv.AddDerivedStatFromExpression(derivedStat.getExpressionStr(), bVerboseErrors);
