[CsvToTelemetry][Analytics] Add option to normalize truth-y strings (true/false/True/False/TRUE/FALSE) to lower via "-normalize-bool-string" flag

diff --git a/Engine/Source/Programs/CSVTools/CSVToTelemetry/Private/CSVToTelemetry.cpp b/Engine/Source/Programs/CSVTools/CSVToTelemetry/Private/CSVToTelemetry.cpp
--- a/Engine/Source/Programs/CSVTools/CSVToTelemetry/Private/CSVToTelemetry.cpp
+++ b/Engine/Source/Programs/CSVTools/CSVToTelemetry/Private/CSVToTelemetry.cpp
@@ -13,6 +13,22 @@
 DEFINE_LOG_CATEGORY_STATIC(LogCSVToTelemetry, Log, All);
 IMPLEMENT_APPLICATION(CSVToTelemetry, "CSVToTelemetry");
 
+FORCEINLINE static bool TryNormalizeBooleanString(FString& InString)
+{
+	if (InString.IsEmpty())
+	{
+		return false;
+	}
+
+	if (InString.Equals(TEXT("true"), ESearchCase::IgnoreCase) || InString.Equals(TEXT("false"), ESearchCase::IgnoreCase))
+	{
+		InString = InString.ToLower();
+		return true;
+	}
+
+	return false;
+}
+
 static int GenerateTelemetryFromCSVFile(const FString& FilePath)
 {
 	// Assume the result will be an error state, non-zero, and only returns zero when fully successful.
@@ -48,7 +64,7 @@ static int GenerateTelemetryFromCSVFile(const FString& FilePath)
 				TArray<FString> KeyArray;
 
 				FString Columns;
-
+				bool NormalizeBooleanStrings = false;
 				int32 SchemaVersion(-1);
 
 				if (FParse::Value(FCommandLine::Get(), TEXT("schema="), SchemaVersion) == false)
@@ -58,17 +74,22 @@ static int GenerateTelemetryFromCSVFile(const FString& FilePath)
 
 				if (FParse::Value(FCommandLine::Get(), TEXT("columns="), Columns) == true)
 				{
-					// We have specified the CSV columns on the command line 
-					Columns.ParseIntoArray(KeyArray, TEXT("|"));
+					// We have specified the CSV columns on the command line; persist empty values to guarantee strict alignment in rows.
+					Columns.ParseIntoArray(KeyArray, TEXT("|"), false);
 					KeySize = Columns.Len();
 				}
 
+				if (FParse::Param(FCommandLine::Get(), TEXT("normalize-bool-strings")))
+				{
+					NormalizeBooleanStrings = true;
+				}
+
 				for (const FString& Row : Rows)
 				{
 					if (KeyArray.IsEmpty())
 					{
-						// Parse the keys on the first row
-						Row.ParseIntoArray(KeyArray, TEXT(","));
+						// Parse the keys on the first row; persist empty values to guarantee strict alignment in rows.
+						Row.ParseIntoArray(KeyArray, TEXT(","), false);
 						KeySize = Row.Len();
 
 						for (int32 i = 0; i < KeyArray.Num(); ++i)
@@ -78,9 +99,9 @@ static int GenerateTelemetryFromCSVFile(const FString& FilePath)
 					}
 					else
 					{
-						// Parse the values from subsequent rows
+						// Parse the values from subsequent rows; persist empty values to match columns.
 						TArray<FString> ValueArray;
-						Row.ParseIntoArray(ValueArray, TEXT(","));
+						Row.ParseIntoArray(ValueArray, TEXT(","), false);
 
 						// Keys and values count should always be the same
 						if (KeyArray.Num() >= ValueArray.Num())
@@ -99,6 +120,13 @@ static int GenerateTelemetryFromCSVFile(const FString& FilePath)
 							for (int32 i = 0; i < ValueArray.Num(); ++i)
 							{
 								ValueArray[i].RemoveSpacesInline();
+
+								// Normalize boolean strings
+								if (NormalizeBooleanStrings)
+								{
+									TryNormalizeBooleanString(ValueArray[i]);
+								}
+
 								Attributes.Emplace(KeyArray[i], ValueArray[i]);
 							}
 
@@ -216,7 +244,7 @@ static int GenerateTelemetryFromCSVProfileFile(const FString& FilePath)
 
 static int ShowHelp()
 {
-	UE_LOG(LogCSVToTelemetry, Display, TEXT("\n\nCSVToTelemetry Help\n\nUsage:\n\tCSVToTelemetry.exe -csv=[filename] -event=[eventname] ( -schema=[value] -column=[name1|name2|....] )\n\tCSVToTelemetry.exe -csvprofile=[filename] -event=[eventname]\n\tCSVToTelemetry.exe -help\n\nRequired:\n\t-csv=[filename]\t\t\tGeneric text based csv input file.\n\t-csvprofile=[filename]\t\tCSVProfiler csv input file. Denote binary with .csv.bin otherwise assumed text.\n\t-event=[name]\t\t\tName of telemetry event to send each row to.\nOptional:\n\t-schema=[value]\t\t\tEvent schema value.\n\t-columns=[name1|name2|....]\tColumn include filter." ) );
+	UE_LOG(LogCSVToTelemetry, Display, TEXT("\n\nCSVToTelemetry Help\n\nUsage:\n\tCSVToTelemetry.exe -csv=[filename] -event=[eventname] ( -schema=[value] -columns=[name1|name2|....] -normalize-bool-strings)\n\tCSVToTelemetry.exe -csvprofile=[filename] -event=[eventname]\n\tCSVToTelemetry.exe -help\n\nRequired:\n\t-csv=[filename]\t\t\tGeneric text based csv input file.\n\t-csvprofile=[filename]\t\tCSVProfiler csv input file. Denote binary with .csv.bin otherwise assumed text.\n\t-event=[name]\t\t\tName of telemetry event to send each row to.\nOptional:\n\t-schema=[value]\t\t\tEvent schema value.\n\t-columns=[name1|name2|....]\tColumn include filter.\n\t-normalize-bool-strings\t\tNormalizes any (\"True\"|\"False\"|\"TRUE\"|\"FALSE\") strings to lowercase."));
 	return 0;
 }
 
@@ -235,14 +263,14 @@ INT32_MAIN_INT32_ARGC_TCHAR_ARGV()
 	int32 Result = 1;
 
 	FString FilePath;
-	
+
 	if (FParse::Value(FCommandLine::Get(), TEXT("help"), FilePath) == true)
 	{
 		// Show help message
 		Result = ShowHelp();
 	}
 	else
-	{ 
+	{
 		if (FParse::Value(FCommandLine::Get(), TEXT("csvprofile="), FilePath) == true)
 		{
 			Result = GenerateTelemetryFromCSVProfileFile(FilePath);
@@ -251,7 +279,7 @@ INT32_MAIN_INT32_ARGC_TCHAR_ARGV()
 		{
 			Result = GenerateTelemetryFromCSVFile(FilePath);
 		}
-	
+
 		if (Result == 0)
 		{
 			// Upload completed successfully
@@ -264,7 +292,7 @@ INT32_MAIN_INT32_ARGC_TCHAR_ARGV()
 		// Always show the usage example if we have not been successful
 		UE_LOG(LogCSVToTelemetry, Error, TEXT("\nUsage:\n\tCSVToTelemetry.exe -csv=[filename] -event=[eventname] ( -schema=[value] -columns=[name1|name2|....] )\n\tCSVToTelemetry.exe -csvprofile=[filename] -event=[eventname]\n\tCSVToTelemetry.exe -help"));
 	}
-	
+
 	if (FParse::Param(FCommandLine::Get(), TEXT("fastexit")))
 	{
 		FPlatformMisc::RequestExitWithStatus(true, Result);
