CsvSplit 1.33. Fix a bug where passing a -delay value would throw an exception

diff --git a/Engine/Source/Programs/CSVTools/CSVSplit/CSVSplit.cs b/Engine/Source/Programs/CSVTools/CSVSplit/CSVSplit.cs
--- a/Engine/Source/Programs/CSVTools/CSVSplit/CSVSplit.cs
+++ b/Engine/Source/Programs/CSVTools/CSVSplit/CSVSplit.cs
@@ -14,7 +14,7 @@ namespace CSVTools
 {
     class Version
     {
-        private static string VersionString = "1.32";
+        private static string VersionString = "1.33";
 
         public static string Get() { return VersionString; }
     };
@@ -84,7 +84,13 @@ namespace CSVTools
             List<int> SourceFrameToDestFrameIndex = new List<int>();
             List<int> SourceFrameToDestCsvIndex = new List<int>();
 
-            for (int i = delay; i < switchStatSamples.samples.Count; i++)
+			for (int i = 0; i < delay; i++)
+			{
+				SourceFrameToDestFrameIndex.Add(0);
+				SourceFrameToDestCsvIndex.Add(0);
+			}
+
+			for (int i = delay; i < switchStatSamples.samples.Count; i++)
             {
                 float switchValue = switchStatSamples.samples[i - delay];
 
