PerfReportTool 4.253.5: Always read csvid metadata as a string. Very occasionally we can end up with an ID which parses as a double and we want to avoid that

diff --git a/Engine/Source/Programs/CSVTools/PerfReportTool/PerfReportTool.cs b/Engine/Source/Programs/CSVTools/PerfReportTool/PerfReportTool.cs
--- a/Engine/Source/Programs/CSVTools/PerfReportTool/PerfReportTool.cs
+++ b/Engine/Source/Programs/CSVTools/PerfReportTool/PerfReportTool.cs
@@ -23,7 +23,7 @@ namespace PerfReportTool
     class Version
     {
 		// Format: Major.Minor.Bugfix
-        private static string VersionString = "4.253.4";
+        private static string VersionString = "4.253.5";
 
         public static string Get() { return VersionString; }
     };
diff --git a/Engine/Source/Programs/CSVTools/PerfReportTool/SummaryTableData.cs b/Engine/Source/Programs/CSVTools/PerfReportTool/SummaryTableData.cs
--- a/Engine/Source/Programs/CSVTools/PerfReportTool/SummaryTableData.cs
+++ b/Engine/Source/Programs/CSVTools/PerfReportTool/SummaryTableData.cs
@@ -532,11 +532,18 @@ namespace PerfSummaries
 		{
 			string key = name.ToLower();
 			double numericValue = double.MaxValue;
-			try
+
+			// Always treat csvid as string
+			bool bIsCsvId = type == SummaryTableElement.Type.CsvMetadata && name == "csvid";
+			if (!bIsCsvId)
 			{
-				numericValue = Convert.ToDouble(value, System.Globalization.CultureInfo.InvariantCulture);
+				try
+				{
+					numericValue = Convert.ToDouble(value, System.Globalization.CultureInfo.InvariantCulture);
+				}
+				catch { }
 			}
-			catch { }
+
 
 			SummaryTableElement metadataValue = null;
 			if (numericValue != double.MaxValue)
