PerfReportTool 4.254.0: Add support for a adding custom tooltips for columns. This works via an optional toolTip attribute in the <statDisplayNameMapping> <mapping> element.

#DisabledValidations Preflight, checkCopyright, checkUnacceptableWords

diff --git a/Engine/Source/Programs/CSVTools/PerfReportTool/PerfReportTool.cs b/Engine/Source/Programs/CSVTools/PerfReportTool/PerfReportTool.cs
--- a/Engine/Source/Programs/CSVTools/PerfReportTool/PerfReportTool.cs
+++ b/Engine/Source/Programs/CSVTools/PerfReportTool/PerfReportTool.cs
@@ -23,7 +23,7 @@ namespace PerfReportTool
     class Version
     {
 		// Format: Major.Minor.Bugfix
-        private static string VersionString = "4.253.5";
+        private static string VersionString = "4.254.0";
 
         public static string Get() { return VersionString; }
     };
@@ -226,7 +226,7 @@ namespace PerfReportTool
 			"  -noCsvCacheFiles: disables usage of .csv.cache files. Cache files can be much faster if filtering on metadata\n" +
 			"";
 
-		Dictionary<string, string> statDisplaynameMapping;
+		Dictionary<string, StatDisplayInfo> statDisplayNameMapping;
 		ReportXML reportXML;
 
 		string GetBaseDirectory()
@@ -442,7 +442,7 @@ namespace PerfReportTool
 				}
 				return;
 			}
-			statDisplaynameMapping = reportXML.GetDisplayNameMapping();
+			statDisplayNameMapping = reportXML.GetDisplayNameMapping();
 
 			// If we're outputting row data to json, create the dict
 			SummaryTableDataJsonWriteHelper summaryTableJsonHelper = null;
@@ -923,7 +923,7 @@ namespace PerfReportTool
 			}
 			else
 			{
-				filteredTable.ApplyDisplayNameMapping(statDisplaynameMapping);
+				filteredTable.ApplyDisplayNameMapping(statDisplayNameMapping);
 				string VersionString = GetBoolArg("noWatermarks") ? "" : Version.Get();
 				string summaryTitle = GetArg("summaryTitle", null);
 				if (GetBoolArg("addDiffRows"))
@@ -1590,7 +1590,7 @@ namespace PerfReportTool
 
 				if (reportTypeInfo.metadataToShowList != null)
 				{
-					Dictionary<string, string> displayNameMapping = reportXML.GetDisplayNameMapping();
+					Dictionary<string, StatDisplayInfo> displayNameMapping = reportXML.GetDisplayNameMapping();
 
 					foreach (string metadataStr in reportTypeInfo.metadataToShowList)
 					{
@@ -1598,9 +1598,12 @@ namespace PerfReportTool
 						if (value != null)
 						{
 							string friendlyName = metadataStr;
-							if (displayNameMapping.ContainsKey(metadataStr.ToLower()))
+							if (displayNameMapping.TryGetValue(metadataStr.ToLower(), out StatDisplayInfo displayInfo))
 							{
-								friendlyName = displayNameMapping[metadataStr];
+								if (displayInfo.displayName != null)
+								{
+									friendlyName = displayNameMapping[metadataStr].displayName;
+								}
 							}
 							htmlFile.WriteLine("<tr><td bgcolor='#F0F0F0'>" + friendlyName + "</td><td><b>" + value + "</b></td></tr>");
 						}
diff --git a/Engine/Source/Programs/CSVTools/PerfReportTool/ReportXml.cs b/Engine/Source/Programs/CSVTools/PerfReportTool/ReportXml.cs
--- a/Engine/Source/Programs/CSVTools/PerfReportTool/ReportXml.cs
+++ b/Engine/Source/Programs/CSVTools/PerfReportTool/ReportXml.cs
@@ -286,7 +286,7 @@ namespace PerfReportTool
 
 
 			// Read the display name mapping
-			statDisplayNameMapping = new Dictionary<string, string>();
+			statDisplayNameMapping = new Dictionary<string, StatDisplayInfo>();
 			XElement displayNameElement = rootElement.Element("statDisplayNameMappings");
 			if (displayNameElement != null)
 			{
@@ -294,9 +294,11 @@ namespace PerfReportTool
 				{
 					string statName = mapping.GetSafeAttribute<string>("statName");
 					string displayName = mapping.GetSafeAttribute<string>("displayName");
-					if (statName != null && displayName != null)
+					string toolTip = mapping.GetSafeAttribute<string>("toolTip");
+
+					if (statName != null && (displayName != null || toolTip != null))
 					{
-						statDisplayNameMapping.Add(statName.ToLower(), displayName);
+						statDisplayNameMapping.Add(statName.ToLower(), new StatDisplayInfo(displayName, toolTip));
 					}
 				}
 			}
@@ -627,7 +629,7 @@ namespace PerfReportTool
 		}
 
 
-		public Dictionary<string, string> GetDisplayNameMapping() { return statDisplayNameMapping; }
+		public Dictionary<string, StatDisplayInfo> GetDisplayNameMapping() { return statDisplayNameMapping; }
 
 		public SummaryTableInfo GetSummaryTable(string name)
 		{
@@ -681,7 +683,7 @@ namespace PerfReportTool
 		string defaultReportTypeName;
 		Dictionary<string, XElement> sharedSummaries;
 		Dictionary<string, GraphSettings> graphs;
-		Dictionary<string, string> statDisplayNameMapping;
+		Dictionary<string, StatDisplayInfo> statDisplayNameMapping;
 		Dictionary<string, string> metadataProxyInfo;
 		public SummaryTableColumnFormatInfoCollection columnFormatInfoList;
 		string baseXmlDirectory;
diff --git a/Engine/Source/Programs/CSVTools/PerfReportTool/SummaryTable.cs b/Engine/Source/Programs/CSVTools/PerfReportTool/SummaryTable.cs
--- a/Engine/Source/Programs/CSVTools/PerfReportTool/SummaryTable.cs
+++ b/Engine/Source/Programs/CSVTools/PerfReportTool/SummaryTable.cs
@@ -12,6 +12,17 @@ using System.Xml.Linq;
 
 namespace PerfSummaries
 {
+	class StatDisplayInfo
+	{
+		public StatDisplayInfo(string inDisplayName, string inToolTip = null)
+		{
+			displayName = inDisplayName;
+			toolTip = inToolTip;
+		}
+		public string displayName = null;
+		public string toolTip = null;
+	};
+
 	class TableUtil
 	{
 		public static string FormatStatName(string inStatName)
@@ -1037,16 +1048,15 @@ namespace PerfSummaries
 			{
 				attributes.Add($"style='background-color:{Colour.White * columnColourModifier};'");
 			}
+			// Generate the tooltip. This consists of a the raw name + element type, followed by a custom provided tooltip, if it exists
+			string tooltipStr = SummaryTable.GetBaseStatNameWithPrefixAndSuffix(name, out _, out _);
+			tooltipStr += " (" + this.elementType.ToString() + ")";
 			if (tooltip != null)
 			{
-				attributes.Add($"title='{tooltip}'");
-			}
-			else 
-			{
-				string tooltipStr = SummaryTable.GetBaseStatNameWithPrefixAndSuffix(name, out _, out _);
-				tooltipStr += " (" + this.elementType.ToString() + ")";
-				attributes.Add($"title='{tooltipStr}'");
+				// Append the custom tooltip with a line break between
+				tooltipStr += "&#013;" + tooltip;
 			}
+			attributes.Add($"title='{tooltipStr}'");
 			return attributes;
 		}
 
@@ -1923,18 +1933,25 @@ namespace PerfSummaries
 			rowWeightings = null;
 		}
 
-		public void ApplyDisplayNameMapping(Dictionary<string, string> statDisplaynameMapping)
+		// Gives columns a display-friendly name and add a tooltip, if provided
+		public void ApplyDisplayNameMapping(Dictionary<string, StatDisplayInfo> statDisplayNameMapping)
 		{
-			// Convert to a display-friendly name
+			if (statDisplayNameMapping == null)
+			{
+				return;
+			}
 			foreach (SummaryTableColumn column in columns)
 			{
-				if (statDisplaynameMapping != null && column.displayName == null)
+				string statName = GetBaseStatNameWithPrefixAndSuffix(column.name, out string prefix, out string suffix);
+				if (statDisplayNameMapping.TryGetValue(statName.ToLower(), out StatDisplayInfo statDisplayInfo))
 				{
-					string name = column.name;
-					string statName = GetBaseStatNameWithPrefixAndSuffix(name, out string prefix, out string suffix);
-					if (statDisplaynameMapping.ContainsKey(statName.ToLower()))
+					if (column.displayName == null && statDisplayInfo.displayName != null)
+					{
+						column.displayName = prefix + statDisplayInfo.displayName + suffix;
+					}
+					if (column.tooltip == null && statDisplayInfo.toolTip != null)
 					{
-						column.displayName = prefix + statDisplaynameMapping[statName.ToLower()] + suffix;
+						column.tooltip = statDisplayInfo.toolTip;
 					}
 				}
 			}
