PerfReportTool: fix inconsistent case-insensitive dict lookup

#DisabledValidations Preflight

diff --git a/Engine/Source/Programs/CSVTools/PerfReportTool/PerfReportTool.cs b/Engine/Source/Programs/CSVTools/PerfReportTool/PerfReportTool.cs
--- a/Engine/Source/Programs/CSVTools/PerfReportTool/PerfReportTool.cs
+++ b/Engine/Source/Programs/CSVTools/PerfReportTool/PerfReportTool.cs
@@ -1602,7 +1602,7 @@ namespace PerfReportTool
 							{
 								if (displayInfo.displayName != null)
 								{
-									friendlyName = displayNameMapping[metadataStr].displayName;
+									friendlyName = displayInfo.displayName;
 								}
 							}
 							htmlFile.WriteLine("<tr><td bgcolor='#F0F0F0'>" + friendlyName + "</td><td><b>" + value + "</b></td></tr>");
