Moved CsvToTelemetry row tokenization over to the CsvParser to improve robustness, support edge cases and encourage code re-use.
Existing code did not support legitimate fields with comments and commas in them for example "Nice, Juicy, Oranges", "Lemons", "Pears" would be incorrectly tokenized to "Nice", "Juicy", "Oranges", "Lemons", "Pears"
No Preflight currently covers buidling of CSVToTelemetry
Added test case for Null field entry based on feedback.

diff --git a/Engine/Source/Programs/CSVTools/CSVToTelemetry/Private/CSVToTelemetry.cpp b/Engine/Source/Programs/CSVTools/CSVToTelemetry/Private/CSVToTelemetry.cpp
--- a/Engine/Source/Programs/CSVTools/CSVToTelemetry/Private/CSVToTelemetry.cpp
+++ b/Engine/Source/Programs/CSVTools/CSVToTelemetry/Private/CSVToTelemetry.cpp
@@ -7,6 +7,7 @@
 #include "AnalyticsET.h"
 #include "StudioTelemetry.h"
 #include "Logging/LogMacros.h"
+#include "Serialization/Csv/CsvParser.h"
 #include "ProjectUtilities.h"
 
 
@@ -49,18 +50,30 @@ static int GenerateTelemetryFromCSVFile(const FString& FilePath)
 		// Start the telemetry session and record each row as a single telemetry event 
 		if (FStudioTelemetry::Get().IsSessionRunning())
 		{
-			TArray<FString> Rows;
+			FString CSVString;
 
-			// Load all the rows from the file into memory
-			if (FFileHelper::LoadFileToStringArray(Rows, *FilePath))
-			{
-				UE_LOG(LogCSVToTelemetry, Display, TEXT("Imported %d rows from file"), Rows.Num());
+			// Load all the file into memory
+			if (FFileHelper::LoadFileToString(CSVString, *FilePath))
+			{	
+				// Parse the file using the CSV Parser
+				FCsvParser Parser(CSVString);
+				const FCsvParser::FRows& Rows = Parser.GetRows();
 
+				if ( Rows.Num()==0 )
+				{
+					// The CSV file failed to parse
+					UE_LOG(LogCSVToTelemetry, Warning, TEXT("Failed to import rows from file %s"), *FilePath);;
+					return -1;
+				}
+				else
+				{
+					UE_LOG(LogCSVToTelemetry, Display, TEXT("Imported %d rows from file %s"), Rows.Num(), *FilePath);
+				}
+				
 				uint32 TotalUploadSize = 0;
-				uint32 RowIndex = 0;
+				uint32 TotalEventsSent = 0;
 				uint32 KeySize = 0;
 				uint32 ValueSize = 0;
-
 				TArray<FString> KeyArray;
 
 				FString Columns;
@@ -72,41 +85,45 @@ static int GenerateTelemetryFromCSVFile(const FString& FilePath)
 					SchemaVersion = -1;
 				}
 
+				if (FParse::Param(FCommandLine::Get(), TEXT("normalize-bool-strings")))
+				{
+					NormalizeBooleanStrings = true;
+				}
+	
 				if (FParse::Value(FCommandLine::Get(), TEXT("columns="), Columns) == true)
 				{
 					// We have specified the CSV columns on the command line; persist empty values to guarantee strict alignment in rows.
 					Columns.ParseIntoArray(KeyArray, TEXT("|"), false);
-					KeySize = Columns.Len();
+					KeySize = Columns.NumBytesWithoutNull();
 				}
 
-				if (FParse::Param(FCommandLine::Get(), TEXT("normalize-bool-strings")))
+				for ( int32 RowIndex=0; RowIndex < Rows.Num(); ++RowIndex )
 				{
-					NormalizeBooleanStrings = true;
-				}
+					const TArray<const TCHAR*>& Row = Rows[RowIndex];
 
-				for (const FString& Row : Rows)
-				{
 					if (KeyArray.IsEmpty())
 					{
-						// Parse the keys on the first row; persist empty values to guarantee strict alignment in rows.
-						Row.ParseIntoArray(KeyArray, TEXT(","), false);
-						KeySize = Row.Len();
-
-						for (int32 i = 0; i < KeyArray.Num(); ++i)
+						for (const TCHAR* Column : Row)
 						{
-							KeyArray[i].RemoveSpacesInline();
+							KeyArray.Push(Column);
 						}
+
+						KeySize = Row.NumBytes();
 					}
 					else
 					{
-						// Parse the values from subsequent rows; persist empty values to match columns.
 						TArray<FString> ValueArray;
-						Row.ParseIntoArray(ValueArray, TEXT(","), false);
+
+						// Parse the values from subsequent rows; persist empty values to match columns.
+						for (const TCHAR* Column : Row)
+						{
+							ValueArray.Push(Column);
+						}
 
 						// Keys and values count should always be the same
 						if (KeyArray.Num() >= ValueArray.Num())
 						{
-							ValueSize = Row.Len();
+							ValueSize = Row.NumBytes();
 
 							TArray<FAnalyticsEventAttribute> Attributes;
 
@@ -135,6 +152,9 @@ static int GenerateTelemetryFromCSVFile(const FString& FilePath)
 
 							// Now we have a complete event so send it
 							FStudioTelemetry::Get().RecordEvent(EventName, Attributes);
+
+							// Keep track of the event count
+							TotalEventsSent++;
 						}
 						else
 						{
@@ -142,11 +162,10 @@ static int GenerateTelemetryFromCSVFile(const FString& FilePath)
 							UE_LOG(LogCSVToTelemetry, Warning, TEXT("Row %d contains the incorrect value count of %d ( expected %d ) and will be skipped."), RowIndex, ValueArray.Num(), KeyArray.Num());
 						}
 					}
-
-					RowIndex++;
 				}
 
 				UE_LOG(LogCSVToTelemetry, Display, TEXT("Generated %d bytes of event data"), TotalUploadSize);
+				UE_LOG(LogCSVToTelemetry, Display, TEXT("Uploading %d %s event(s)..."), TotalEventsSent, *EventName);
 
 				// Everything worked as expected return code should be zero
 				Result = 0;
