Added support for additional event attributes to be passed on the CSVToTelemetry commandline with -attributes=[key1=value1|key2=value2|...] syntax.

Addional attributes are sent with every event.

Added SchemaVersion set by the -schema=[version] option to the additional attributes which simplifies the code. Note : This -schema=[version] can also be achieved with -attributes=SchemaVersion=version but have left this option intact for backward compatabilty.

diff --git a/Engine/Source/Programs/CSVTools/CSVToTelemetry/Private/CSVToTelemetry.cpp b/Engine/Source/Programs/CSVTools/CSVToTelemetry/Private/CSVToTelemetry.cpp
--- a/Engine/Source/Programs/CSVTools/CSVToTelemetry/Private/CSVToTelemetry.cpp
+++ b/Engine/Source/Programs/CSVTools/CSVToTelemetry/Private/CSVToTelemetry.cpp
@@ -75,26 +75,61 @@ static int GenerateTelemetryFromCSVFile(const FString& FilePath)
 				uint32 KeySize = 0;
 				uint32 ValueSize = 0;
 				TArray<FString> KeyArray;
+				TArray<FAnalyticsEventAttribute> AdditionalAttributes;
 
-				FString Columns;
 				bool NormalizeBooleanStrings = false;
 				int32 SchemaVersion(-1);
 
-				if (FParse::Value(FCommandLine::Get(), TEXT("schema="), SchemaVersion) == false)
+				if (FParse::Value(FCommandLine::Get(), TEXT("schema="), SchemaVersion) == true)
 				{
-					SchemaVersion = -1;
+					// Add our schema version
+					AdditionalAttributes.Emplace(TEXT("SchemaVersion"), SchemaVersion);
 				}
 
 				if (FParse::Param(FCommandLine::Get(), TEXT("normalize-bool-strings")))
 				{
 					NormalizeBooleanStrings = true;
 				}
-	
-				if (FParse::Value(FCommandLine::Get(), TEXT("columns="), Columns) == true)
+
+				FString ColumnsString;
+				if (FParse::Value(FCommandLine::Get(), TEXT("columns="), ColumnsString) == true)
 				{
 					// We have specified the CSV columns on the command line; persist empty values to guarantee strict alignment in rows.
-					Columns.ParseIntoArray(KeyArray, TEXT("|"), false);
-					KeySize = Columns.NumBytesWithoutNull();
+					ColumnsString.ParseIntoArray(KeyArray, TEXT("|"), false);
+					KeySize = ColumnsString.NumBytesWithoutNull();
+				}
+				
+				FString AttributesString;
+				if (FParse::Value(FCommandLine::Get(), TEXT("attributes="), AttributesString) == true)
+				{
+					// Parse a list of additional attributes to pass to the each event
+					TArray<FString> AttributesArray;
+					AttributesString.ParseIntoArray(AttributesArray, TEXT("|"), false);
+
+					if (AttributesArray.Num() > 0)
+					{
+						for (int32 AttributeIndex = 0; AttributeIndex < AttributesArray.Num(); ++AttributeIndex)
+						{
+							TArray<FString> KeyValuePairArray;
+							AttributesArray[AttributeIndex].ParseIntoArray(KeyValuePairArray, TEXT("="), false);
+
+							if (KeyValuePairArray.Num() == 2)
+							{
+								AdditionalAttributes.Emplace(KeyValuePairArray[0], KeyValuePairArray[1]);
+							}
+							else
+							{
+								UE_LOG(LogCSVToTelemetry, Error, TEXT("Invalid additional attributes \"%s\". They must be a valid Key Value Pair Key=Value"), *AttributesArray[AttributeIndex]);
+								return Result;
+							}
+						}
+
+						UE_LOG(LogCSVToTelemetry, Display, TEXT("Appending %d additional attributes"), AdditionalAttributes.Num());
+					}
+					else
+					{
+						UE_LOG(LogCSVToTelemetry, Warning, TEXT("No additional attributes found in %s"), *AttributesString);
+					}
 				}
 
 				for ( int32 RowIndex=0; RowIndex < Rows.Num(); ++RowIndex )
@@ -125,13 +160,7 @@ static int GenerateTelemetryFromCSVFile(const FString& FilePath)
 						{
 							ValueSize = Row.NumBytes();
 
-							TArray<FAnalyticsEventAttribute> Attributes;
-
-							if (SchemaVersion != -1)
-							{
-								// Add our schema version
-								Attributes.Emplace(TEXT("SchemaVersion"), SchemaVersion);
-							}
+							TArray<FAnalyticsEventAttribute> Attributes = AdditionalAttributes;
 
 							// Add the Key/Value pair to the attribute list
 							for (int32 i = 0; i < ValueArray.Num(); ++i)
@@ -263,7 +292,7 @@ static int GenerateTelemetryFromCSVProfileFile(const FString& FilePath)
 
 static int ShowHelp()
 {
-	UE_LOG(LogCSVToTelemetry, Display, TEXT("\n\nCSVToTelemetry Help\n\nUsage:\n\tCSVToTelemetry.exe -csv=[filename] -event=[eventname] ( -schema=[value] -columns=[name1|name2|....] -normalize-bool-strings)\n\tCSVToTelemetry.exe -csvprofile=[filename] -event=[eventname]\n\tCSVToTelemetry.exe -help\n\nRequired:\n\t-csv=[filename]\t\t\tGeneric text based csv input file.\n\t-csvprofile=[filename]\t\tCSVProfiler csv input file. Denote binary with .csv.bin otherwise assumed text.\n\t-event=[name]\t\t\tName of telemetry event to send each row to.\nOptional:\n\t-schema=[value]\t\t\tEvent schema value.\n\t-columns=[name1|name2|....]\tColumn include filter.\n\t-normalize-bool-strings\t\tNormalizes any (\"True\"|\"False\"|\"TRUE\"|\"FALSE\") strings to lowercase."));
+	UE_LOG(LogCSVToTelemetry, Display, TEXT("\n\nCSVToTelemetry Help\n\nUsage:\n\tCSVToTelemetry.exe -csv=[filename] -event=[eventname] ( -schema=[value] -columns=[name1|name2|....] -attributes=[key1=value1|key2=value2|....] -normalize-bool-strings)\n\tCSVToTelemetry.exe -csvprofile=[filename] -event=[eventname]\n\tCSVToTelemetry.exe -help\n\nRequired:\n\t-csv=[filename]\t\t\tGeneric text based csv input file.\n\t-csvprofile=[filename]\t\tCSVProfiler csv input file. Denote binary with .csv.bin otherwise assumed text.\n\t-event=[name]\t\t\tName of telemetry event to send each row to.\nOptional:\n\t-schema=[value]\t\t\tEvent schema value.\n\t-columns=[name1|name2|....]\tColumn include filter.\n\t-normalize-bool-strings\t\tNormalizes any (\"True\"|\"False\"|\"TRUE\"|\"FALSE\") strings to lowercase."));
 	return 0;
 }
 
@@ -309,7 +338,7 @@ INT32_MAIN_INT32_ARGC_TCHAR_ARGV()
 	if (Result != 0)
 	{
 		// Always show the usage example if we have not been successful
-		UE_LOG(LogCSVToTelemetry, Error, TEXT("\nUsage:\n\tCSVToTelemetry.exe -csv=[filename] -event=[eventname] ( -schema=[value] -columns=[name1|name2|....] )\n\tCSVToTelemetry.exe -csvprofile=[filename] -event=[eventname]\n\tCSVToTelemetry.exe -help"));
+		UE_LOG(LogCSVToTelemetry, Error, TEXT("\nUsage:\n\tCSVToTelemetry.exe -csv=[filename] -event=[eventname] ( -schema=[value] -columns=[name1|name2|....] -attributes=[key1=value1|key2=value2|....] )\n\tCSVToTelemetry.exe -csvprofile=[filename] -event=[eventname]\n\tCSVToTelemetry.exe -help"));
 	}
 
 	if (FParse::Param(FCommandLine::Get(), TEXT("fastexit")))
