PerfReportTool 4.255: Add support for ShowTotals in graphs

diff --git a/Engine/Source/Programs/CSVTools/PerfReportTool/PerfReportTool.cs b/Engine/Source/Programs/CSVTools/PerfReportTool/PerfReportTool.cs
--- a/Engine/Source/Programs/CSVTools/PerfReportTool/PerfReportTool.cs
+++ b/Engine/Source/Programs/CSVTools/PerfReportTool/PerfReportTool.cs
@@ -23,7 +23,7 @@ namespace PerfReportTool
     class Version
     {
 		// Format: Major.Minor.Bugfix
-        private static string VersionString = "4.254.0";
+        private static string VersionString = "4.255.0";
 
         public static string Get() { return VersionString; }
     };
@@ -1873,6 +1873,7 @@ namespace PerfReportTool
 			int height = graphSettings.height.value;
 			bool stacked = graphSettings.stacked.value;
 			bool showAverages = graphSettings.showAverages.value;
+			bool showTotals = graphSettings.showTotals.value;
 			bool filterOutZeros = graphSettings.filterOutZeros.value;
 			bool snapToPeaks = false;
 			if (graphSettings.snapToPeaks.isSet)
@@ -1976,6 +1977,7 @@ namespace PerfReportTool
 			AddConditionalFlag("stacked", stacked);
 			AddConditionalFlag("forceLegendSort", stacked); // based on the stacked flag
 			AddConditionalFlag("showAverages", showAverages);
+			AddConditionalFlag("showTotals", showTotals);
 			AddConditionalFlag("nosnap", !snapToPeaks);
 			AddConditionalFlag("filterOutZeros", filterOutZeros);
 			AddConditionalArg("maxHierarchyDepth", maxHierarchyDepth > 0, maxHierarchyDepth);
@@ -2080,6 +2082,10 @@ namespace PerfReportTool
 			{
 				graphParams.showAverages = graphSettings.showAverages.value;
 			}
+			if (graphSettings.showTotals.isSet)
+			{
+				graphParams.showTotals = graphSettings.showTotals.value;
+			}
 			if (graphSettings.filterOutZeros.isSet)
 			{
 				graphParams.filterOutZeros = graphSettings.filterOutZeros.value;
diff --git a/Engine/Source/Programs/CSVTools/PerfReportTool/ReportGraph.cs b/Engine/Source/Programs/CSVTools/PerfReportTool/ReportGraph.cs
--- a/Engine/Source/Programs/CSVTools/PerfReportTool/ReportGraph.cs
+++ b/Engine/Source/Programs/CSVTools/PerfReportTool/ReportGraph.cs
@@ -194,6 +194,7 @@ namespace PerfReportTool
 			height = new Optional<int>(element, "height", vars);
 			stacked = new Optional<bool>(element, "stacked", vars);
 			showAverages = new Optional<bool>(element, "showAverages", vars);
+			showTotals = new Optional<bool>(element, "showTotals", vars);
 			filterOutZeros = new Optional<bool>(element, "filterOutZeros", vars);
 			maxHierarchyDepth = new Optional<int>(element, "maxHierarchyDepth", vars);
 			hideStatPrefix = new OptionalString(element, "hideStatPrefix", false, vars);
@@ -236,6 +237,7 @@ namespace PerfReportTool
 			//additionalArgs.InheritFrom(baseSettings.additionalArgs);
 			stacked.InheritFrom(baseSettings.stacked);
 			showAverages.InheritFrom(baseSettings.showAverages);
+			showTotals.InheritFrom(baseSettings.showTotals);
 			filterOutZeros.InheritFrom(baseSettings.filterOutZeros);
 			maxHierarchyDepth.InheritFrom(baseSettings.maxHierarchyDepth);
 			hideStatPrefix.InheritFrom(baseSettings.hideStatPrefix);
@@ -271,6 +273,7 @@ namespace PerfReportTool
 		//public OptionalString additionalArgs;
 		public Optional<bool> stacked;
 		public Optional<bool> showAverages;
+		public Optional<bool> showTotals;
 		public Optional<bool> filterOutZeros;
 		public Optional<int> maxHierarchyDepth;
 		public OptionalString hideStatPrefix;
