PerfReportTool 4.256.0: Add a standard summary which is always present. For now this just computes the "Total Time (s)" column if it doesn't already exist (this is the case for report types which don't include the FPSCharts summary). This fixes issues with weighted average computation in summary tables, since many tables use this column to compute a weighted average.

diff --git a/Engine/Source/Programs/CSVTools/PerfReportTool/PerfReportTool.cs b/Engine/Source/Programs/CSVTools/PerfReportTool/PerfReportTool.cs
--- a/Engine/Source/Programs/CSVTools/PerfReportTool/PerfReportTool.cs
+++ b/Engine/Source/Programs/CSVTools/PerfReportTool/PerfReportTool.cs
@@ -23,7 +23,7 @@ namespace PerfReportTool
     class Version
     {
 		// Format: Major.Minor.Bugfix
-        private static string VersionString = "4.255.0";
+        private static string VersionString = "4.256.0";
 
         public static string Get() { return VersionString; }
     };
@@ -1639,6 +1639,8 @@ namespace PerfReportTool
 				bool bLinkTemplates = GetBoolArg("linkTemplates");
 				summaries.Insert(0, new ExtraLinksSummary(null, reportTypeInfo.vars, null, bLinkTemplates));
 			}
+			// Add the standard summary last
+			summaries.Add(new StandardSummary());
 
 			// If the reporttype has summary info, then write out the summary]
 			foreach (Summary summary in summaries)
diff --git a/Engine/Source/Programs/CSVTools/PerfReportTool/Summaries/StandardSummary.cs b/Engine/Source/Programs/CSVTools/PerfReportTool/Summaries/StandardSummary.cs
new file mode 100644
--- /dev/null
+++ b/Engine/Source/Programs/CSVTools/PerfReportTool/Summaries/StandardSummary.cs
@@ -0,0 +1,53 @@
+// Copyright Epic Games, Inc. All Rights Reserved.
+
+using CSVStats;
+using PerfReportTool;
+using System;
+using System.Collections.Generic;
+using System.Data.Common;
+using System.IO;
+using System.Linq;
+using System.Xml.Linq;
+
+namespace PerfSummaries
+{
+	class StandardSummary : Summary
+	{
+		/*
+		  A standard summary which is present in all report types. Requires no XML definition. Executed after all other summaries
+		*/
+	
+		public StandardSummary()
+		{
+		}
+
+		public override string GetName() { return "standard"; }
+
+
+		public override HtmlSection WriteSummaryData(bool bWriteHtml, CsvStats csvStats, CsvStats csvStatsUnstripped, bool bWriteSummaryCsv, SummaryTableRowData rowData, string htmlFileName)
+		{
+			HtmlSection htmlSection = null;
+
+			// Compute the total time if the RowData doesn't already exist
+			if (rowData != null)
+			{
+				string TotalTimeColumnName = "Total Time (s)";
+				if (!rowData.Contains(TotalTimeColumnName.ToLower()))
+				{
+					if (csvStats.Stats.TryGetValue("frametime", out StatSamples frameTimes))
+					{
+						double totalTimeMs = 0.0;
+						foreach (float frameTime in frameTimes.samples)
+						{
+							totalTimeMs += (double)frameTime;
+						}
+						rowData.Add(SummaryTableElement.Type.SummaryTableMetric, TotalTimeColumnName, totalTimeMs / 1000.0);
+					}
+				}
+			}
+			return htmlSection;
+		}
+
+	};
+
+}
\ No newline at end of file
