PerfReportTool 4.257.0: Add a sortByName option for the events summary table

diff --git a/Engine/Source/Programs/CSVTools/PerfReportTool/PerfReportTool.cs b/Engine/Source/Programs/CSVTools/PerfReportTool/PerfReportTool.cs
--- a/Engine/Source/Programs/CSVTools/PerfReportTool/PerfReportTool.cs
+++ b/Engine/Source/Programs/CSVTools/PerfReportTool/PerfReportTool.cs
@@ -23,7 +23,7 @@ namespace PerfReportTool
     class Version
     {
 		// Format: Major.Minor.Bugfix
-        private static string VersionString = "4.256.0";
+        private static string VersionString = "4.257.0";
 
         public static string Get() { return VersionString; }
     };
diff --git a/Engine/Source/Programs/CSVTools/PerfReportTool/Summaries/EventSummary.cs b/Engine/Source/Programs/CSVTools/PerfReportTool/Summaries/EventSummary.cs
--- a/Engine/Source/Programs/CSVTools/PerfReportTool/Summaries/EventSummary.cs
+++ b/Engine/Source/Programs/CSVTools/PerfReportTool/Summaries/EventSummary.cs
@@ -18,6 +18,7 @@ namespace PerfSummaries
 			summaryStatName = element.GetRequiredAttribute<string>(vars, "summaryStatName");
 			events = element.Element("events").Value.Split(',');
 			colourThresholds = ReadColourThresholdsXML(element.Element("colourThresholds"), vars);
+			bSortByName = element.GetSafeAttribute<bool>("sortByName", false);
 		}
 
 		public EventSummary() { }
@@ -55,9 +56,18 @@ namespace PerfSummaries
 			}
 
 			// Sort descending
-			eventCountsDict = eventCountsDict
-				.OrderBy(pair  => -pair.Value)
-				.ToDictionary(x => x.Key, x => x.Value);
+			if (bSortByName)
+			{
+				eventCountsDict = eventCountsDict
+					.OrderBy(pair => pair.Key)
+					.ToDictionary(x => x.Key, x => x.Value);
+			}
+			else
+			{
+				eventCountsDict = eventCountsDict
+					.OrderBy(pair => -pair.Value)
+					.ToDictionary(x => x.Key, x => x.Value);
+			}
 
 			// Output HTML
 			if (bWriteHtml && eventCountsDict.Count > 0)
@@ -97,6 +107,7 @@ namespace PerfSummaries
 		double[] colourThresholds;
 		string title;
 		string summaryStatName;
+		bool bSortByName;
 	};
 
 }
\ No newline at end of file
