PerfReportTool/CSVToSvg: Derived stats: Add a cumulative() function to generate cumulative stats from non-cumulative, e.g time=cumulative(frametime)/1000 can be used to graph cumulative time in seconds

diff --git a/Engine/Source/Programs/CSVTools/CsvStats/CsvStats.cs b/Engine/Source/Programs/CSVTools/CsvStats/CsvStats.cs
--- a/Engine/Source/Programs/CSVTools/CsvStats/CsvStats.cs
+++ b/Engine/Source/Programs/CSVTools/CsvStats/CsvStats.cs
@@ -573,7 +573,7 @@ namespace CSVStats
 			return null;
 		}
 
-		enum BuiltinFunctionType { addstats, meta, eventcount, framedelta, COUNT };
+		enum BuiltinFunctionType { addstats, meta, eventcount, framedelta, cumulative, COUNT };
 		string ReadNextToken(string expression, ref int index, Dictionary<string, StatSamples> statsOut, int tokenIndex)
 		{
 			// Skip leading whitespace
@@ -689,6 +689,33 @@ namespace CSVStats
 						statsOut[destStatName] = newStat;
 						return destStatName;
 					}
+					else if (funcType == BuiltinFunctionType.cumulative)
+					{
+						// Quotes are optional. They're supported for consistency with stats
+						string srcStatName = ReadQuotedString(paramString);
+						srcStatName = (srcStatName == null) ? paramString : srcStatName;
+						if (srcStatName == null)
+						{
+							throw new DerivedStatException("Cumulative requires an arg: " + expression[oldIndex..index]);
+						}
+
+						StatSamples srcStat = GetStat(srcStatName);
+						if (srcStat == null)
+						{
+							throw new DerivedStatMissingSourceException("Source stat not found: " + expression[oldIndex..index]);
+						}
+						string destStatName = "__STAT__CUMULATIVE__" + tokenIndex.ToString();
+						StatSamples newStat = new StatSamples(destStatName, srcStat.samples.Count);
+						double total = 0.0;
+						for (int j = 1; j < srcStat.samples.Count; j++)
+						{
+							total += srcStat.samples[j];
+							newStat.samples[j] = (float)total;
+						}
+						statsOut[destStatName] = newStat;
+						return destStatName;
+					}
+
 				}
 			}
 
