CsvStats: Fix bug where first frame isn't counted towards cumulative derived CSV stats

diff --git a/Engine/Source/Programs/CSVTools/CsvStats/CsvStats.cs b/Engine/Source/Programs/CSVTools/CsvStats/CsvStats.cs
--- a/Engine/Source/Programs/CSVTools/CsvStats/CsvStats.cs
+++ b/Engine/Source/Programs/CSVTools/CsvStats/CsvStats.cs
@@ -707,7 +707,7 @@ namespace CSVStats
 						string destStatName = "__STAT__CUMULATIVE__" + tokenIndex.ToString();
 						StatSamples newStat = new StatSamples(destStatName, srcStat.samples.Count);
 						double total = 0.0;
-						for (int j = 1; j < srcStat.samples.Count; j++)
+						for (int j = 0; j < srcStat.samples.Count; j++)
 						{
 							total += srcStat.samples[j];
 							newStat.samples[j] = (float)total;
