PerfReportTool 4.258.0 + CsvStats
Add support for CSV event timestamps. Requires bumping the CsvBinary file version
#tests tested via CsvTools unit tests

diff --git a/Engine/Source/Programs/CSVTools/CSVSplit/CSVSplit.cs b/Engine/Source/Programs/CSVTools/CSVSplit/CSVSplit.cs
--- a/Engine/Source/Programs/CSVTools/CSVSplit/CSVSplit.cs
+++ b/Engine/Source/Programs/CSVTools/CSVSplit/CSVSplit.cs
@@ -143,8 +143,8 @@ namespace CSVTools
                 {
                     int destCsvIndex = SourceFrameToDestCsvIndex[eventFrame];
                     CsvStats destStats = csvStatsOutList[destCsvIndex];
-                    CsvEvent newEvent = new CsvEvent(ev.Name, SourceFrameToDestFrameIndex[ev.Frame]);
-                    destStats.Events.Add(ev);
+                    CsvEvent newEvent = new CsvEvent(ev.Name, SourceFrameToDestFrameIndex[ev.Frame], ev.Timestamp);
+                    destStats.Events.Add(newEvent);
 
                     if (bVirtualEvents)
                     {
@@ -153,7 +153,7 @@ namespace CSVTools
                         {
                             if (j != destCsvIndex)
                             {
-                                CsvEvent newEvent2 = new CsvEvent(ev.Name + " [VIRTUAL]", newEvent.Frame);
+                                CsvEvent newEvent2 = new CsvEvent(ev.Name + " [VIRTUAL]", newEvent.Frame, ev.Timestamp);
                                 destStats.Events.Add(newEvent2);
                             }
                         }
diff --git a/Engine/Source/Programs/CSVTools/CsvStats/CsvStats.cs b/Engine/Source/Programs/CSVTools/CsvStats/CsvStats.cs
--- a/Engine/Source/Programs/CSVTools/CsvStats/CsvStats.cs
+++ b/Engine/Source/Programs/CSVTools/CsvStats/CsvStats.cs
@@ -20,6 +20,7 @@ namespace CSVStats
 		PreRelease = 1,
 		InitialRelease,
 		CompressionSupportAndFlags,
+		EventTimestamps,
 
 		COUNT,
 		CURRENT = COUNT - 1
@@ -57,19 +58,24 @@ namespace CSVStats
 		{
 			Name = source.Name;
 			Frame = source.Frame;
+			Timestamp = source.Timestamp;
 		}
-		public CsvEvent(string NameIn, int FrameIn)
+		public CsvEvent(string NameIn, int FrameIn, double TimestampIn)
         {
             Name = NameIn;
             Frame = FrameIn;
+			Timestamp = TimestampIn; 
         }
         public CsvEvent()
         {
 
         }
 
+		public bool HasTimestamp() { return Timestamp != double.MaxValue; }
+
         public string Name;
         public int Frame;
+		public double Timestamp = double.MaxValue;
     };
 
 	public class DerivedStatException : Exception
@@ -1412,6 +1418,7 @@ namespace CSVStats
 			{
 				fileWriter.Write(ev.Frame);
 				fileWriter.Write(ev.Name);
+				fileWriter.Write(ev.Timestamp);
 			}
 			fileWriter.Close();
 		}
@@ -1598,9 +1605,14 @@ namespace CSVStats
 			{
 				int frame = fileReader.ReadInt32();
 				string name = fileReader.ReadString();
+				double timestamp = double.MaxValue;
+				if (version >= (int)CsvBinVersion.EventTimestamps)
+				{
+					timestamp = fileReader.ReadDouble();
+				}
 				if (frame >= numRowsToSkip)
 				{
-					CsvEvent ev = new CsvEvent(name, frame);
+					CsvEvent ev = new CsvEvent(name, frame, timestamp);
 					csvStatsOut.Events.Add(ev);
 				}
 			}
@@ -2125,7 +2137,7 @@ namespace CSVStats
 						CsvEvent currentEvent=Events[eventIndex];
 						if (bFrameExists)
 						{
-							newEvents.Add(new CsvEvent(currentEvent.Name, currentEvent.Frame-skippedFrameCount));
+							newEvents.Add(new CsvEvent(currentEvent.Name, currentEvent.Frame-skippedFrameCount, currentEvent.Timestamp));
 						}
 						eventIndex++;
 					}
@@ -2462,9 +2474,22 @@ namespace CSVStats
                         {
                             if (EventString.Length > 0)
                             {
-                                CsvEvent ev = new CsvEvent();
-                                ev.Frame = i;
-                                ev.Name = EventString;
+								double timestamp = double.MaxValue;
+								string name = EventString;
+
+								// Check if there's a timestamp embedded
+								int suffixIndex = EventString.LastIndexOf("##");
+								if (suffixIndex != -1)
+								{
+									name = EventString.Substring(0, suffixIndex);
+									if ( double.TryParse(EventString.Substring(suffixIndex + 2), NumberStyles.Float, CultureInfo.InvariantCulture, out double result) )
+									{
+										timestamp = result;
+									}
+									
+								}
+
+								CsvEvent ev = new CsvEvent(name, i, timestamp);
                                 FilteredEvents.Add(ev);
                             }
                         }
diff --git a/Engine/Source/Programs/CSVTools/CsvToSvgLib/CsvToSvgLib.cs b/Engine/Source/Programs/CSVTools/CsvToSvgLib/CsvToSvgLib.cs
--- a/Engine/Source/Programs/CSVTools/CsvToSvgLib/CsvToSvgLib.cs
+++ b/Engine/Source/Programs/CSVTools/CsvToSvgLib/CsvToSvgLib.cs
@@ -15,7 +15,7 @@ namespace CSVTools
 
 	public class CsvToSvgLibVersion
 	{
-		private static string VersionString = "3.69";
+		private static string VersionString = "3.70";
 
 		public static string Get() { return VersionString; }
 	};
@@ -1674,6 +1674,7 @@ namespace CSVTools
 			{
 				base.Frame = ev.Frame;
 				base.Name = ev.Name;
+				base.Timestamp = ev.Timestamp;
 				count = 1;
 				priority = inPriority;
 			}
diff --git a/Engine/Source/Programs/CSVTools/PerfReportTool/PerfReportTool.cs b/Engine/Source/Programs/CSVTools/PerfReportTool/PerfReportTool.cs
--- a/Engine/Source/Programs/CSVTools/PerfReportTool/PerfReportTool.cs
+++ b/Engine/Source/Programs/CSVTools/PerfReportTool/PerfReportTool.cs
@@ -23,7 +23,7 @@ namespace PerfReportTool
     class Version
     {
 		// Format: Major.Minor.Bugfix
-        private static string VersionString = "4.257.0";
+        private static string VersionString = "4.258.0";
 
         public static string Get() { return VersionString; }
     };
