Fix a bug where event timestamps would be stripped when writing out text CSVs (e.g with CsvConvert)

diff --git a/Engine/Source/Programs/CSVTools/CsvStats/CsvStats.cs b/Engine/Source/Programs/CSVTools/CsvStats/CsvStats.cs
--- a/Engine/Source/Programs/CSVTools/CsvStats/CsvStats.cs
+++ b/Engine/Source/Programs/CSVTools/CsvStats/CsvStats.cs
@@ -1625,17 +1625,18 @@ namespace CSVStats
             System.IO.StreamWriter csvOutFile;
             csvOutFile = new System.IO.StreamWriter(filename);
 
-            // Store the events in a dictionary so we can lookup by frame
+            // Store the event text in a dictionary so we can lookup by frame
             Dictionary<int, string> eventsLookup = new Dictionary<int, string>();
             foreach (CsvEvent ev in Events)
             {
+				string eventText = ev.HasTimestamp() ? ev.Name + "##" + ev.Timestamp.ToString(CultureInfo.InvariantCulture) : ev.Name;
                 if (!eventsLookup.ContainsKey(ev.Frame))
                 {
-                    eventsLookup.Add(ev.Frame, ev.Name);
+                    eventsLookup.Add(ev.Frame, eventText);
                 }
                 else
                 {
-                    eventsLookup[ev.Frame] += ";" + ev.Name;
+                    eventsLookup[ev.Frame] += ";" + eventText;
                 }
             }
 
