Added check for an empty or badly formed CSV file where the first row contained empty values.

There should either be a non-empty set of keys on the first row of the file or keys are specified using the -column= option.

diff --git a/Engine/Source/Programs/CSVTools/CSVToTelemetry/Private/CSVToTelemetry.cpp b/Engine/Source/Programs/CSVTools/CSVToTelemetry/Private/CSVToTelemetry.cpp
--- a/Engine/Source/Programs/CSVTools/CSVToTelemetry/Private/CSVToTelemetry.cpp
+++ b/Engine/Source/Programs/CSVTools/CSVToTelemetry/Private/CSVToTelemetry.cpp
@@ -132,64 +132,90 @@ static int GenerateTelemetryFromCSVFile(const FString& FilePath)
 					}
 				}
 
-				for ( int32 RowIndex=0; RowIndex < Rows.Num(); ++RowIndex )
+				int32 RowIndex = 0;
+
+				if (KeyArray.IsEmpty())
 				{
+					// The user has not specified any keys so assume these are on the first row
 					const TArray<const TCHAR*>& Row = Rows[RowIndex];
 
-					if (KeyArray.IsEmpty())
+					for (const TCHAR* Column : Row)
 					{
-						for (const TCHAR* Column : Row)
+						KeyArray.Push(Column);
+					}
+
+					KeySize = Row.NumBytes();
+
+					RowIndex++;
+				}
+
+				// Validate that we have a valid set of keys. These are either specified on the first row of the csv file or explicitly via the -columns= commandline
+				if (KeyArray.IsEmpty())
+				{
+					// No keys we found at all
+					UE_LOG(LogCSVToTelemetry, Error, TEXT("No keys are specified for csv file %s"), *FilePath);
+					return Result;
+				}
+				else
+				{
+					// Keys were found but validate that any are not empty
+					for (const FString& Key : KeyArray)
+					{
+						if (Key.IsEmpty())
 						{
-							KeyArray.Push(Column);
+							// Keys can't ever be empty
+							UE_LOG(LogCSVToTelemetry, Error, TEXT("Empty key specifed in csv file %s"), *FilePath);
+							return Result;
 						}
+					}
+				}
+
+				for ( ; RowIndex < Rows.Num(); ++RowIndex )
+				{
+					const TArray<const TCHAR*>& Row = Rows[RowIndex];
+
+					TArray<FString> ValueArray;
 
-						KeySize = Row.NumBytes();
+					// Parse the values from subsequent rows; persist empty values to match columns.
+					for (const TCHAR* Column : Row)
+					{
+						ValueArray.Push(Column);
 					}
-					else
+
+					// Keys and values count should always be the same
+					if (KeyArray.Num() >= ValueArray.Num())
 					{
-						TArray<FString> ValueArray;
+						ValueSize = Row.NumBytes();
 
-						// Parse the values from subsequent rows; persist empty values to match columns.
-						for (const TCHAR* Column : Row)
-						{
-							ValueArray.Push(Column);
-						}
+						TArray<FAnalyticsEventAttribute> Attributes = AdditionalAttributes;
 
-						// Keys and values count should always be the same
-						if (KeyArray.Num() >= ValueArray.Num())
+						// Add the Key/Value pair to the attribute list
+						for (int32 i = 0; i < ValueArray.Num(); ++i)
 						{
-							ValueSize = Row.NumBytes();
-
-							TArray<FAnalyticsEventAttribute> Attributes = AdditionalAttributes;
+							ValueArray[i].RemoveSpacesInline();
 
-							// Add the Key/Value pair to the attribute list
-							for (int32 i = 0; i < ValueArray.Num(); ++i)
+							// Normalize boolean strings
+							if (NormalizeBooleanStrings)
 							{
-								ValueArray[i].RemoveSpacesInline();
-
-								// Normalize boolean strings
-								if (NormalizeBooleanStrings)
-								{
-									TryNormalizeBooleanString(ValueArray[i]);
-								}
-
-								Attributes.Emplace(KeyArray[i], ValueArray[i]);
+								TryNormalizeBooleanString(ValueArray[i]);
 							}
 
-							// Keep track of the total data we're uploading
-							TotalUploadSize += KeySize + ValueSize;
+							Attributes.Emplace(KeyArray[i], ValueArray[i]);
+						}
 
-							// Now we have a complete event so send it
-							FStudioTelemetry::Get().RecordEvent(EventName, Attributes);
+						// Keep track of the total data we're uploading
+						TotalUploadSize += KeySize + ValueSize;
 
-							// Keep track of the event count
-							TotalEventsSent++;
-						}
-						else
-						{
-							// The key/value counts don't match
-							UE_LOG(LogCSVToTelemetry, Warning, TEXT("Row %d contains the incorrect value count of %d ( expected %d ) and will be skipped."), RowIndex, ValueArray.Num(), KeyArray.Num());
-						}
+						// Now we have a complete event so send it
+						FStudioTelemetry::Get().RecordEvent(EventName, Attributes);
+
+						// Keep track of the event count
+						TotalEventsSent++;
+					}
+					else
+					{
+						// The key/value counts don't match
+						UE_LOG(LogCSVToTelemetry, Warning, TEXT("Row %d contains the incorrect value count of %d ( expected %d ) and will be skipped."), RowIndex, ValueArray.Num(), KeyArray.Num());
 					}
 				}
 
