Added -maxrows=[value] commandline option to limit the number of rows exported by the tool
Will export the full number of rows found in the csv filre or the maximum num whichever is the smallest
This option is mostly useful for debugging purposes where you are testing schema changes and don't need to send a full data set

diff --git a/Engine/Source/Programs/CSVTools/CSVToTelemetry/Private/CSVToTelemetry.cpp b/Engine/Source/Programs/CSVTools/CSVToTelemetry/Private/CSVToTelemetry.cpp
--- a/Engine/Source/Programs/CSVTools/CSVToTelemetry/Private/CSVToTelemetry.cpp
+++ b/Engine/Source/Programs/CSVTools/CSVToTelemetry/Private/CSVToTelemetry.cpp
@@ -74,6 +74,7 @@ static int GenerateTelemetryFromCSVFile(const FString& FilePath)
 				uint32 TotalEventsSent = 0;
 				uint32 KeySize = 0;
 				uint32 ValueSize = 0;
+				uint32 MaxRows = -1;
 				TArray<FString> KeyArray;
 				TArray<FAnalyticsEventAttribute> AdditionalAttributes;
 
@@ -91,6 +92,8 @@ static int GenerateTelemetryFromCSVFile(const FString& FilePath)
 					NormalizeBooleanStrings = true;
 				}
 
+				FParse::Value(FCommandLine::Get(), TEXT("maxrows="), MaxRows);
+				
 				FString ColumnsString;
 				if (FParse::Value(FCommandLine::Get(), TEXT("columns="), ColumnsString) == true)
 				{
@@ -132,7 +135,7 @@ static int GenerateTelemetryFromCSVFile(const FString& FilePath)
 					}
 				}
 
-				int32 RowIndex = 0;
+				uint32 RowIndex = 0;
 
 				if (KeyArray.IsEmpty())
 				{
@@ -170,7 +173,10 @@ static int GenerateTelemetryFromCSVFile(const FString& FilePath)
 					}
 				}
 
-				for ( ; RowIndex < Rows.Num(); ++RowIndex )
+				// Limit the maximum number of rows or choose the full set
+				MaxRows = FMath::Min<uint32>(Rows.Num(), MaxRows);
+				
+				for ( ; RowIndex < MaxRows; ++RowIndex )
 				{
 					const TArray<const TCHAR*>& Row = Rows[RowIndex];
 
@@ -318,7 +324,7 @@ static int GenerateTelemetryFromCSVProfileFile(const FString& FilePath)
 
 static int ShowHelp()
 {
-	UE_LOG(LogCSVToTelemetry, Display, TEXT("\n\nCSVToTelemetry Help\n\nUsage:\n\tCSVToTelemetry.exe -csv=[filename] -event=[eventname] ( -schema=[value] -columns=[name1|name2|....] -attributes=[key1=value1|key2=value2|....] -normalize-bool-strings)\n\tCSVToTelemetry.exe -csvprofile=[filename] -event=[eventname]\n\tCSVToTelemetry.exe -help\n\nRequired:\n\t-csv=[filename]\t\t\tGeneric text based csv input file.\n\t-csvprofile=[filename]\t\tCSVProfiler csv input file. Denote binary with .csv.bin otherwise assumed text.\n\t-event=[name]\t\t\tName of telemetry event to send each row to.\nOptional:\n\t-schema=[value]\t\t\tEvent schema value.\n\t-columns=[name1|name2|....]\tColumn include filter.\n\t-normalize-bool-strings\t\tNormalizes any (\"True\"|\"False\"|\"TRUE\"|\"FALSE\") strings to lowercase."));
+	UE_LOG(LogCSVToTelemetry, Display, TEXT("\n\nCSVToTelemetry Help\n\nUsage:\n\tCSVToTelemetry.exe -csv=[filename] -event=[eventname] ( -schema=[value] -columns=[name1|name2|....] -attributes=[key1=value1|key2=value2|....] -normalize-bool-strings -maxrows=[value])\n\tCSVToTelemetry.exe -csvprofile=[filename] -event=[eventname]\n\tCSVToTelemetry.exe -help\n\nRequired:\n\t-csv=[filename]\t\t\tGeneric text based csv input file.\n\t-csvprofile=[filename]\t\tCSVProfiler csv input file. Denote binary with .csv.bin otherwise assumed text.\n\t-event=[name]\t\t\tName of telemetry event to send each row to.\nOptional:\n\t-schema=[value]\t\t\tEvent schema value.\n\t-columns=[name1|name2|....]\tColumn include filter.\n\t-normalize-bool-strings\t\tNormalizes any (\"True\"|\"False\"|\"TRUE\"|\"FALSE\") strings to lowercase."));
 	return 0;
 }
 
