PerfReportTool 4.259.0: Add support for CSV timing regions via the Timing region summary. This visualizes region concurrency and outputs a metric for each showing the sum of each.
Example XML definition:
<summary name="BootToLobbyTimingRegions" type="timingregion" title="Boot to lobby timing regions" endevent="BootPhase/Lobby" metricPrefix="BootRegionTimes">
<regionFilter pattern="BootPhase/*"/>
<regionFilter pattern="Interactability/*"/>
<regionFilter pattern="GFP/*"/>
<regionFilter pattern="OnlineAccount/*" busyTimeMetric="LoginTime" hidePrefix="Login:"/>
<regionFilter pattern="LoadingScreen"/>
<regionFilter pattern="*"/>
</summary>
#tests PRT test suite

diff --git a/Engine/Source/Programs/CSVTools/CsvStats/CsvStats.cs b/Engine/Source/Programs/CSVTools/CsvStats/CsvStats.cs
--- a/Engine/Source/Programs/CSVTools/CsvStats/CsvStats.cs
+++ b/Engine/Source/Programs/CSVTools/CsvStats/CsvStats.cs
@@ -529,6 +529,31 @@ namespace CSVStats
 			return metaData.GetValue(key, null);
 		}
 
+		public CsvEvent FindEvent(string pattern, double startTime = 0.0, double endTime = double.MaxValue)
+		{
+			foreach (CsvEvent ev in Events)
+			{
+				if (CsvStats.DoesSearchStringMatch(ev.Name, pattern) && ev.Timestamp >= startTime && ev.Timestamp <= endTime)
+				{
+					return ev;
+				}
+			}
+			return null;
+		}
+
+		public List<CsvEvent> FindEvents(string pattern, double startTime = 0.0, double endTime = double.MaxValue)
+		{
+			List<CsvEvent> eventList = new List<CsvEvent>();
+			foreach (CsvEvent ev in Events)
+			{
+				if (CsvStats.DoesSearchStringMatch(ev.Name, pattern) && ev.Timestamp >= startTime && ev.Timestamp <= endTime)
+				{
+					eventList.Add(ev);
+				}
+			}
+			return eventList;
+		}
+
 		// Matches 
 		string ReadQuotedString(string str)
 		{
@@ -1110,6 +1135,18 @@ namespace CSVStats
 			}
 		}
 
+		public static string GetColorFromString(string input, int minBrightness = 64, int maxBrightness = 200)
+		{
+			using (var md5 = MD5.Create())
+			{
+				byte[] hash = md5.ComputeHash(Encoding.UTF8.GetBytes(input));
+				byte r = (byte)(hash[0] % (maxBrightness - minBrightness) + minBrightness);
+				byte g = (byte)(hash[1] % (maxBrightness - minBrightness) + minBrightness);
+				byte b = (byte)(hash[2] % (maxBrightness - minBrightness) + minBrightness);
+				return $"#{r:X2}{g:X2}{b:X2}";
+			}
+		}
+
 		public static List<string> WildcardMatchStringList(List<string> stringList, string searchString, bool bCaseSensitive = false, bool bListIsSorted = false)
 		{
 			searchString = searchString.Trim();
diff --git a/Engine/Source/Programs/CSVTools/PerfReportTool/PerfReportTool.cs b/Engine/Source/Programs/CSVTools/PerfReportTool/PerfReportTool.cs
--- a/Engine/Source/Programs/CSVTools/PerfReportTool/PerfReportTool.cs
+++ b/Engine/Source/Programs/CSVTools/PerfReportTool/PerfReportTool.cs
@@ -23,7 +23,7 @@ namespace PerfReportTool
     class Version
     {
 		// Format: Major.Minor.Bugfix
-        private static string VersionString = "4.258.0";
+        private static string VersionString = "4.259.0";
 
         public static string Get() { return VersionString; }
     };
diff --git a/Engine/Source/Programs/CSVTools/PerfReportTool/Summaries/TimingRegionSummary.cs b/Engine/Source/Programs/CSVTools/PerfReportTool/Summaries/TimingRegionSummary.cs
new file mode 100644
--- /dev/null
+++ b/Engine/Source/Programs/CSVTools/PerfReportTool/Summaries/TimingRegionSummary.cs
@@ -0,0 +1,511 @@
+// Copyright Epic Games, Inc. All Rights Reserved.
+
+using CSVStats;
+using PerfReportTool;
+using System;
+using System.Collections.Generic;
+using System.Xml.Linq;
+using System.Linq;
+
+namespace PerfSummaries
+{
+	class TimingRegionSummary : Summary
+	{
+
+		class RegionFilterInfo
+		{
+			public RegionFilterInfo(XElement element, XmlVariableMappings vars, int indexIn)
+			{
+				pattern = element.GetRequiredAttribute<string>("pattern");
+				busyTimeMetric = element.GetSafeAttribute<string>("busyTimeMetric");
+				hidePrefix = element.GetSafeAttribute<string>("hidePrefix");
+				index = indexIn;
+			}
+			public string pattern;
+			public string busyTimeMetric = null;
+			public string hidePrefix = null;
+			public int index = 0;
+		}
+		public TimingRegionSummary(XElement element, XmlVariableMappings vars, string baseXmlDirectory)
+		{
+			ReadStatsFromXML(element, vars);
+			if (stats.Count != 0)
+			{
+				throw new Exception("<stats> element is not supported");
+			}
+			title = element.GetRequiredAttribute<string>(vars, "title");
+			startEventName = element.GetSafeAttribute<string>(vars, "startEvent");
+			endEventName = element.GetSafeAttribute<string>(vars, "endEvent");
+			metricPrefix = element.GetSafeAttribute<string>(vars, "metricPrefix");
+			foreach (XElement regionFilterXml in element.Elements("regionFilter"))
+			{
+				RegionFilterInfo filterInfo = new RegionFilterInfo(regionFilterXml, vars, regionFilters.Count);
+				regionFilters.Add(filterInfo);
+			}
+		}
+		public TimingRegionSummary() { }
+
+		public override string GetName() { return "timingregion"; }
+
+		enum ERangeType
+		{
+			Before,
+			Inside,
+			After
+		}
+
+		class RegionEvent
+		{
+			public static RegionEvent TryCreateFromCsvEvent(CsvEvent csvEvent) 
+			{
+				const string regionPrefix = "Region: ";
+				if (!csvEvent.Name.StartsWith(regionPrefix))
+				{
+					return null;
+				}
+				const string beginSuffix = " - Begin";
+				const string endSuffix = " - End";
+				RegionEvent regionEvent = new RegionEvent();
+				string name = csvEvent.Name.Substring(regionPrefix.Length);
+				if (name.EndsWith(beginSuffix))
+				{
+					regionEvent.name = name.Substring(0, name.Length - beginSuffix.Length);
+					regionEvent.bBegin = true;
+				}
+				else if (name.EndsWith(endSuffix))
+				{
+					regionEvent.name = name.Substring(0, name.Length - endSuffix.Length);
+					regionEvent.bBegin = false;
+				}
+				else
+				{
+					// No begin or end
+					return null;
+				}
+
+				// Read the category if there is one
+				int slashIndex = regionEvent.name.IndexOf('/');
+				if (slashIndex != -1)
+				{
+					regionEvent.category = regionEvent.name.Substring(0, slashIndex);
+					regionEvent.name = regionEvent.name.Substring(slashIndex+1);
+				}
+				else
+				{
+					// Can't have null as the category since we use it as a dictionary key
+					regionEvent.category = "";
+				}
+				regionEvent.timestamp = csvEvent.Timestamp;
+				regionEvent.frameNumber = csvEvent.Frame;
+				return regionEvent;
+			}
+
+			public double GetPairDuration()
+			{ 
+				if (otherEvent == null)
+				{
+					return 0.0;
+				}
+				return Math.Abs(timestamp - otherEvent.timestamp);
+			}
+
+			public string GetFullName()
+			{
+				if (category == "")
+				{
+					return name;
+				}
+				return category + "/" + name;
+			}
+
+			public RegionEvent CreateTruncatedPairEvent(double startTime, double endTime)
+			{
+				// If this is an end event, make sure it's the last one before any begin events
+				RegionEvent newEvent = new RegionEvent();
+				newEvent.bBegin = !bBegin;
+				newEvent.name = name;
+				newEvent.category = category;
+				newEvent.frameNumber = -1;
+				newEvent.timestamp = bBegin ? endTime : startTime;
+				newEvent.otherEvent = this;
+				newEvent.rangeType = bBegin ? ERangeType.After : ERangeType.Before;
+				newEvent.filterInfo = filterInfo;
+				otherEvent = newEvent;
+				return newEvent;
+			}
+
+			public string category;
+			public string name;
+			public int frameNumber = -1;
+			public bool bBegin;
+			public double timestamp;
+			public RegionEvent otherEvent;
+			public RegionFilterInfo filterInfo;
+			public ERangeType rangeType = ERangeType.Inside;
+		}
+
+		bool TryGetRegionFilterInfo(RegionEvent regionEvent, out RegionFilterInfo outRegionFilterInfo)
+		{
+			outRegionFilterInfo = null;
+			foreach (RegionFilterInfo filterInfo in regionFilters)
+			{
+				if (CsvStats.DoesSearchStringMatch(regionEvent.GetFullName(), filterInfo.pattern))
+				{
+					outRegionFilterInfo = filterInfo;
+					return true;
+				}
+			}
+			return false;
+		}
+
+
+
+		List<RegionEvent> GetRegionEvents(CsvStats csvStats, double startTime, double endTime)
+		{
+			// Find the region events
+			List<RegionEvent> allRegionEvents = new List<RegionEvent>();
+
+			// Make sure the events are sorted by timestamp
+			List<CsvEvent> csvEvents = new List<CsvEvent>(csvStats.Events);
+			csvEvents.Sort((a, b) => a.Timestamp.CompareTo(b.Timestamp));
+
+			foreach (CsvEvent ev in csvStats.Events)
+			{
+				RegionEvent regionEvent = RegionEvent.TryCreateFromCsvEvent(ev);
+				if (regionEvent != null && TryGetRegionFilterInfo(regionEvent, out regionEvent.filterInfo))
+				{
+					// Detect and clamp out of range events
+					if (ev.Timestamp < startTime)
+					{
+						regionEvent.rangeType = ERangeType.Before;
+					}
+					else if (ev.Timestamp > endTime)
+					{
+						regionEvent.rangeType = ERangeType.After;
+					}
+					regionEvent.timestamp = Math.Clamp(ev.Timestamp, startTime, endTime);
+					allRegionEvents.Add(regionEvent);
+				}
+			}
+
+			// Split the events by name
+			Dictionary<string, List<RegionEvent>> regionEventsByName = new Dictionary<string, List<RegionEvent>>();
+			foreach (RegionEvent ev in allRegionEvents)
+			{
+				string key = ev.GetFullName();
+				List<RegionEvent> regionEventList;
+				if (!regionEventsByName.TryGetValue(key, out regionEventList))
+				{
+					regionEventList = new List<RegionEvent>();
+					regionEventsByName[key] = regionEventList;
+				}
+				regionEventList.Add(ev);
+			}
+
+			List<RegionEvent> allFilteredRegionEvents = new List<RegionEvent>();
+			// Spin through each region by name
+			foreach (List<RegionEvent> regionList in regionEventsByName.Values)
+			{
+				// Compute the depth of every event, including orphan events which span the boundary
+				int depth = 0;
+				int minDepth = int.MaxValue;
+
+				Dictionary<RegionEvent, int> regionEventDepth = new Dictionary<RegionEvent, int>();
+				foreach (RegionEvent ev in regionList) 
+				{
+					if (ev.bBegin)
+					{
+						regionEventDepth[ev] = depth;
+						minDepth = Math.Min(minDepth, depth);
+						depth++;
+					}
+					else
+					{
+						depth--;
+						minDepth = Math.Min(minDepth, depth);
+						regionEventDepth[ev] = depth;
+					}
+				}
+
+				// Cull everything not at min depth and pair the adjacent begin/end nodes
+				List<RegionEvent> filteredEvents = new List<RegionEvent>();
+				RegionEvent beginEvent = null;
+				foreach (RegionEvent ev in regionList)
+				{
+					if (regionEventDepth[ev] == minDepth)
+					{
+						if (ev.bBegin)
+						{
+							beginEvent = ev;
+						}
+						else if (beginEvent != null)
+						{
+							ev.otherEvent = beginEvent;
+							beginEvent.otherEvent = ev;
+							beginEvent = null;
+						}
+						filteredEvents.Add(ev);
+					}
+				}
+
+				// Detect orphan events at the start and end and add dummy events connecting to them
+				if (filteredEvents.Count > 0)
+				{
+					if (filteredEvents.First().otherEvent == null)
+					{
+						filteredEvents.Insert(0, filteredEvents.First().CreateTruncatedPairEvent(startTime, endTime));
+					}
+					if (filteredEvents.Last().otherEvent == null)
+					{
+						filteredEvents.Add(filteredEvents.Last().CreateTruncatedPairEvent(startTime, endTime));
+					}
+				}
+				allFilteredRegionEvents.AddRange(filteredEvents);
+			}
+
+			// Cull region events which are entirely outside the range on one side or the other or the duration is very close to zero
+			List<RegionEvent> newRegionEvents = new List<RegionEvent>();
+			foreach (RegionEvent ev in allFilteredRegionEvents)
+			{
+				if (ev.rangeType == ERangeType.Inside || ev.otherEvent.rangeType == ERangeType.Inside)
+				{
+					if ( ev.GetPairDuration() > 0.0001 )
+					{
+						newRegionEvents.Add(ev);
+					}
+				}
+			}
+
+			newRegionEvents.Sort((a, b) => a.timestamp.CompareTo(b.timestamp));
+			return newRegionEvents;
+		}
+
+		class RegionEventTrack
+		{
+			public RegionEventTrack(string inCategory, double inStartTime, double inEndTime, RegionFilterInfo inFilterInfo)
+			{
+				category = inCategory;
+				startTime = inStartTime;
+				endTime = inEndTime;
+				filterInfo = inFilterInfo;
+				filterIndex = filterInfo.index;
+			}
+			public bool TryAddEvent(RegionEvent beginEvent)
+			{
+				if (beginEvent.category != category || beginEvent.filterInfo != filterInfo || beginEvent.timestamp < lastEventTime)
+				{
+					return false;
+				}
+				lastEventTime = (beginEvent.otherEvent != null) ? beginEvent.otherEvent.timestamp : endTime;
+				events.Add(beginEvent);
+				return true;
+			}
+
+			public void WriteToHtml(HtmlSection htmlSection, double xOffset, double yOffset, double trackWidth, double trackHeight, string backgroundColor="#ffffff")
+			{
+				// Strip the hidePrefix if necessary
+				string backgroundStyle = "fill:" + backgroundColor;
+				string textStyle = "font-family:verdana;font-size:12px";
+				double trackDuration = endTime - startTime;
+				double lastEventX=xOffset;
+				foreach (RegionEvent beginEvent in events)
+				{
+					RegionEvent endEvent = beginEvent.otherEvent;
+					double duration = endEvent.timestamp - beginEvent.timestamp;
+					double t = beginEvent.timestamp - startTime;
+					double x = Math.Round(xOffset + trackWidth * t / trackDuration,0);
+					double width = Math.Round(trackWidth * duration / trackDuration,0);
+
+					// Truncate the display name if required
+					string displayName = beginEvent.name;
+					if (filterInfo.hidePrefix != null && beginEvent.name.StartsWith(filterInfo.hidePrefix))
+					{
+						displayName = beginEvent.name.Substring(filterInfo.hidePrefix.Length);
+					}
+
+					string regionText = displayName + " - " + String.Format("{0:0.00}", duration) + "s";
+					string style = "fill:"+CsvStats.GetColorFromString(displayName, 80, 220)+";stroke-width:1;stroke:black";
+					string titleStr = "<title>" + regionText + "</title>";
+
+					double overlap = 0;
+					if ( Math.Abs(lastEventX - x) < 0.001 )
+					{
+						if ( lastEventX > 0.0)
+						{
+							overlap = 1; // Overlap consecutive blocks to avoid double border lines
+						}
+					}
+					else
+					{
+						// Draw a white block before this event to cover any overflow text
+						htmlSection.WriteLine($"<rect x='{lastEventX}' y='{yOffset}' width='{x - lastEventX}' height='{trackHeight}' style='{backgroundStyle}'/>");
+					}
+					htmlSection.WriteLine($"<rect x='{(x - overlap)}' y='{yOffset}' width='{width + overlap}' height='{trackHeight}' style='{style}'>");
+					htmlSection.WriteLine(titleStr+"</rect>");
+					if (width > 20)
+					{
+						htmlSection.WriteLine($"<text x='{x + 3}' y='{yOffset + 15}' style='{textStyle}' pointer-events='none'>{titleStr + regionText}</text>");
+					}
+					lastEventX = x + width;
+				}
+				// Redraw the background on the RHS to deal with any overflow text
+				htmlSection.WriteLine($"<rect x='{(lastEventX+1)}' y='{yOffset}' width='{(trackWidth-(lastEventX-xOffset)-1)}' height='{trackHeight}' style='{backgroundStyle}'/>");
+			}
+
+			public string category;
+			public int filterIndex;
+			public double lastEventTime;
+			public List<RegionEvent> events = new List<RegionEvent>();
+			public double startTime;
+			public double endTime;
+			public RegionFilterInfo filterInfo;
+		}
+
+		public override HtmlSection WriteSummaryData(bool bWriteHtml, CsvStats csvStats, CsvStats csvStatsUnstripped, bool bWriteSummaryCsv, SummaryTableRowData rowData, string htmlFileName)
+		{
+			HtmlSection htmlSection = null;
+
+			StatSamples frameTimeStat = csvStatsUnstripped.GetStat("frametime");
+			if (frameTimeStat == null)
+			{
+				throw new Exception("TimingRegion summary requires a frametime stat");
+			}
+
+			CsvEvent startEvent = (startEventName != null) ? csvStatsUnstripped.FindEvent(startEventName) : null;
+			double startTime = (startEvent != null) ? startEvent.Timestamp : 0.0;
+
+			CsvEvent endEvent = (endEventName != null) ? csvStatsUnstripped.FindEvent(endEventName, startTime) : null;
+			double endTime = (endEvent != null) ? endEvent.Timestamp : frameTimeStat.total / 1000.0;
+
+			double totalDuration = endTime - startTime;
+
+			List<RegionEvent> regionEvents = GetRegionEvents(csvStatsUnstripped, startTime, endTime);
+			if ( regionEvents.Count == 0 )
+			{
+				// No relevant events found
+				return null;
+			}
+
+			// Assign events to tracks
+			List<RegionEventTrack> tracks = new List<RegionEventTrack>();
+			foreach (RegionEvent ev in regionEvents)
+			{
+				if (ev.bBegin && ev.otherEvent != null)
+				{
+					// Add the event to an available track
+					bool bFoundTrack = false;
+					foreach (RegionEventTrack track in tracks)
+					{
+						if (track.TryAddEvent(ev))
+						{
+							bFoundTrack = true;
+							break;
+						}
+					}
+					if ( !bFoundTrack )
+					{
+						RegionEventTrack newTrack = new RegionEventTrack(ev.category, startTime, endTime, ev.filterInfo);
+						if ( newTrack.TryAddEvent(ev) )
+						{
+							tracks.Add(newTrack);
+						}
+					}
+				}
+			}
+
+			// Sort tracks by filter index
+			tracks.Sort((a, b) => a.filterIndex.CompareTo(b.filterIndex));
+
+			// Output HTML
+			string prevCategory = null;
+			if (bWriteHtml)
+			{
+				htmlSection = new HtmlSection(title, bStartCollapsed);
+
+				double imageWidth = 1400;
+				double imageHeight = tracks.Count * 22;
+				double headerWidth = 150;
+				double trackWidth = imageWidth - headerWidth;
+
+				string lineStyle = "stroke:#808080;stroke-width:2";
+				string categoryTextStyle = "font-family:verdana;font-size:12px";
+
+				List<double> trackWidths = new List<double>();
+
+				// Check if the file exists in the output directory
+				htmlSection.WriteLine($"<svg version='1.1' xmlns='http://www.w3.org/2000/svg' xmlns:xlink='http://www.w3.org/1999/xlink' width='{imageWidth}' height='{imageHeight}'>");
+
+				double yOffset = 0;
+				int trackIndex = 0;
+				foreach (RegionEventTrack track in tracks)
+				{
+					if (prevCategory != track.category)
+					{
+						if (prevCategory != null)
+						{
+							double lineY = yOffset + 1;
+							htmlSection.WriteLine($"<line x1='{0}' y1='{lineY}' x2='{imageWidth}' y2='{lineY}' style='{lineStyle}' stroke-dasharray='5,5'/>");
+							yOffset += 2;
+						}
+						prevCategory = track.category;
+						htmlSection.WriteLine($"<text x='{0}' y='{yOffset+14}' style='{categoryTextStyle}'>{track.category}</text>");
+					}
+					track.WriteToHtml(htmlSection, headerWidth, yOffset, trackWidth, 20, trackIndex%2==0 ? "#e0e0e0" : "#f0f0f0");
+					yOffset += 20;
+					trackIndex++;
+				}
+				htmlSection.WriteLine($"<line x1='{headerWidth-1}' y1='{0}' x2='{headerWidth-1}' y2='{imageHeight}' style='{lineStyle}' stroke-dasharray='5,5'/>");
+
+				htmlSection.WriteLine("</svg>");
+			}
+
+			// Output row data showing the total time for each region
+			Dictionary<string, double> regionTotalTimes = new Dictionary<string, double>();
+			foreach (RegionEvent ev in regionEvents)
+			{
+				if ( ev.bBegin )
+				{
+					if (!regionTotalTimes.ContainsKey(ev.GetFullName()))
+					{
+						regionTotalTimes[ev.GetFullName()] = 0;
+					}
+					regionTotalTimes[ev.GetFullName()] += ev.otherEvent.timestamp - ev.timestamp;
+				}
+			}
+
+			if (metricPrefix != null && rowData != null)
+			{
+				foreach (KeyValuePair<string, double> pair in regionTotalTimes)
+				{
+					string name = metricPrefix + "/" + pair.Key;
+					rowData.Add(new SummaryTableElement(SummaryTableElement.Type.SummaryTableMetric, name, pair.Value, null, ""));
+				}
+			}
+			if (htmlSection != null)
+			{
+				htmlSection.WriteLine("<h3>Total times</h3>");
+				htmlSection.WriteLine("<table border='0' style='width:1000'>");
+				htmlSection.WriteLine("<tr><th>Region</th><th>Time</th></tr>");
+				List<KeyValuePair<string, double>> regionTimingsSorted = new List<KeyValuePair<string, double>>(regionTotalTimes);
+				regionTimingsSorted.Sort((a, b) => a.Key.CompareTo(b.Key));
+
+				foreach (KeyValuePair<string, double> pair in regionTimingsSorted)
+				{
+					htmlSection.WriteLine($"<tr><td>{pair.Key}</td><td>{String.Format("{0:0.00}", pair.Value)}</td></tr>");
+				}
+				htmlSection.WriteLine("</table>");
+			}
+			return htmlSection;
+		}
+		public override void PostInit(ReportTypeInfo reportTypeInfo, CsvStats csvStats)
+		{
+		}
+
+		string title;
+		string startEventName;
+		string endEventName;
+		string metricPrefix;
+		List<RegionFilterInfo> regionFilters = new List<RegionFilterInfo>();
+	};
+
+}
\ No newline at end of file
diff --git a/Engine/Source/Programs/CSVTools/PerfReportTool/Summary.cs b/Engine/Source/Programs/CSVTools/PerfReportTool/Summary.cs
--- a/Engine/Source/Programs/CSVTools/PerfReportTool/Summary.cs
+++ b/Engine/Source/Programs/CSVTools/PerfReportTool/Summary.cs
@@ -35,7 +35,7 @@ namespace PerfSummaries
 				object[] constructArgs = new object[3] { summaryXmlElement, vars, baseXmlDirectory };
 				return (Summary)Activator.CreateInstance(summaryType, constructArgs);
 			}
-			throw new Exception("Summary type " + summaryType + " not found!");
+			throw new Exception("Summary type " + summaryTypeName + " not found!");
 		}
 
 		static Dictionary<string, System.Type> summaryNameLookup = null;
