Add Support for generating zoomed in map images to see more granular detail on large maps. Defaulted to off.
Add Support for multiple start events.

#rnx
#tests Generated maps for multiple experiences with zoomed maps on and off

diff --git a/Engine/Source/Programs/CSVTools/PerfReportTool/PerfReportTool.cs b/Engine/Source/Programs/CSVTools/PerfReportTool/PerfReportTool.cs
--- a/Engine/Source/Programs/CSVTools/PerfReportTool/PerfReportTool.cs
+++ b/Engine/Source/Programs/CSVTools/PerfReportTool/PerfReportTool.cs
@@ -23,7 +23,7 @@ namespace PerfReportTool
     class Version
     {
 		// Format: Major.Minor.Bugfix
-        private static string VersionString = "4.259.0";
+        private static string VersionString = "4.260.0";
 
         public static string Get() { return VersionString; }
     };
diff --git a/Engine/Source/Programs/CSVTools/PerfReportTool/Summaries/MapOverlaySummary.cs b/Engine/Source/Programs/CSVTools/PerfReportTool/Summaries/MapOverlaySummary.cs
--- a/Engine/Source/Programs/CSVTools/PerfReportTool/Summaries/MapOverlaySummary.cs
+++ b/Engine/Source/Programs/CSVTools/PerfReportTool/Summaries/MapOverlaySummary.cs
@@ -1,15 +1,16 @@
 // Copyright Epic Games, Inc. All Rights Reserved.
 
+using CSVStats;
+using PerfReportTool;
 using System;
-using System.IO;
 using System.Collections.Generic;
-using System.Linq;
-using System.Xml.Linq;
-using PerfReportTool;
-using CSVStats;
 using System.Drawing;
 using System.Drawing.Imaging;
+using System.IO;
+using System.Linq;
+using System.Numerics;
 using System.Runtime.Versioning;
+using System.Xml.Linq;
 
 namespace PerfSummaries
 {
@@ -30,6 +31,12 @@ namespace PerfSummaries
 			public string lineColor;
 		};
 
+		struct SVGPosition
+		{
+			public int X;
+			public int Y;
+		};
+
 		class MapOverlay
 		{
 			public MapOverlay(XElement element, XmlVariableMappings vars)
@@ -47,19 +54,14 @@ namespace PerfSummaries
 					ev.lineColor = eventEl.GetSafeAttribute<string>(vars, "lineColor");
 					if (eventEl.GetSafeAttribute<bool>(vars, "isStartEvent", false))
 					{
-						if (startEvent != null)
-						{
-							throw new Exception("Can't have multiple start events!");
-						}
-						startEvent = ev;
+						startEvents.Add(ev);
 					}
 					events.Add(ev);
 				}
-
 			}
 			public string[] positionStatNames = new string[3];
 			public string summaryStatNamePrefix;
-			public MapOverlayEvent startEvent;
+			public List<MapOverlayEvent> startEvents = new List<MapOverlayEvent>();
 			public string lineColor;
 			public List<MapOverlayEvent> events = new List<MapOverlayEvent>();
 		}
@@ -102,12 +104,28 @@ namespace PerfSummaries
 				transform = Array.ConvertAll(transformtext.Split(' ', StringSplitOptions.RemoveEmptyEntries), float.Parse);
 			}
 
+			string fulltransformtext = element.GetSafeAttribute(vars, "transformFull", "");
+			if (fulltransformtext.Length != 0)
+			{
+				fullTransform = Array.ConvertAll(fulltransformtext.Split(' ', StringSplitOptions.RemoveEmptyEntries), float.Parse);
+				if (fullTransform.Length != 16)
+				{
+					throw new Exception("Malformed transform detected");
+				}
+
+				transformMatrix = new Matrix4x4(fullTransform[0], fullTransform[1], fullTransform[2], fullTransform[3],
+												fullTransform[4], fullTransform[5], fullTransform[6], fullTransform[7],
+												fullTransform[8], fullTransform[9], fullTransform[10], fullTransform[11],
+												fullTransform[12], fullTransform[13], fullTransform[14], fullTransform[15]);
+			}
+
 			title = element.GetSafeAttribute(vars, "title", "Events");
 			destImageFilename = element.GetRequiredAttribute<string>(vars, "destImage");
 			imageWidth = element.GetSafeAttribute<float>(vars, "width", 250.0f);
 			imageHeight = element.GetSafeAttribute<float>(vars, "height", 250.0f);
 			framesPerLineSegment = element.GetSafeAttribute<int>(vars, "framesPerLineSegment", 5);
 			lineSplitDistanceThreshold = element.GetSafeAttribute<float>(vars, "lineSplitDistanceThreshold", float.MaxValue);
+			createCroppedImage = element.GetSafeAttribute<bool>(vars, "createCroppedImage", false);
 
 			foreach (XElement overlayEl in element.Elements("overlay"))
 			{
@@ -136,6 +154,56 @@ namespace PerfSummaries
 			return (int)(svgY + 0.5f);
 		}
 
+		SVGPosition toSVGLocation(float worldX, float worldY, Rectangle imageRegion)
+		{
+			Vector3 svgPos;
+
+			// TODO: Remove when we fully support full transforms
+			if (fullTransform != null)
+			{
+				// These seem backwards?
+				Vector3 worldPos3D = new Vector3(worldY, worldX, 0.0f);
+				svgPos = Vector3.Transform(worldPos3D, transformMatrix);
+
+				svgPos.Y = -svgPos.Y;
+				// From -1<>1 to 0<>1
+				svgPos.X += 1.0f;
+				svgPos.Y += 1.0f;
+				svgPos.X *= 0.5f;
+				svgPos.Y *= 0.5f;
+
+				// To Image Size
+				svgPos.X *= imageWidth;
+				svgPos.Y *= imageHeight;
+			}
+			else
+			{
+				svgPos.X = toSvgX(worldX, worldY);
+				svgPos.Y = toSvgY(worldX, worldY);
+			}
+
+			// Scale to SubImage
+			if (imageRegion.Width != imageWidth || imageRegion.Height != imageHeight)
+			{
+				svgPos.X = (svgPos.X - imageRegion.X) / imageRegion.Width * imageWidth;
+				svgPos.Y = (svgPos.Y - imageRegion.Y) / imageRegion.Height * imageHeight;
+			}
+
+			return new SVGPosition()
+			{
+				X = (int)(svgPos.X + 0.5f),
+				Y =	(int)(svgPos.Y + 0.5f),
+			};
+		}
+
+		SVGPosition toSVGLocationClamped(float worldX, float worldY, Rectangle imageRegion)
+		{
+			SVGPosition svgPos = toSVGLocation(worldX, worldY, imageRegion);
+			svgPos.X = Math.Min(Math.Max(0, svgPos.X), (int)imageWidth);
+			svgPos.Y = Math.Min(Math.Max(0, svgPos.Y), (int)imageHeight);
+			return svgPos;
+		}
+
 		private void CopyAndResizeImage(string sourceImagePath, string destImagePath, int destWidth, int destHeight)
 		{
 			if (!OperatingSystem.IsWindows())
@@ -147,7 +215,7 @@ namespace PerfSummaries
 			using (FileStream fileStream = new FileStream(sourceImagePath, FileMode.Open, FileAccess.Read))
 			{
 				Console.WriteLine("Reading source image");
-				using (var image = Image.FromStream(fileStream))
+				using (var image = System.Drawing.Image.FromStream(fileStream))
 				{
 					Console.WriteLine("Generating downsampled image");
 					var thumbnail = image.GetThumbnailImage(destWidth, destHeight, null, IntPtr.Zero);
@@ -160,6 +228,59 @@ namespace PerfSummaries
 			}
 		}
 
+		private void CopyAndTrimImage(string sourceImagePath, string destImagePath, int destWidth, int destHeight, Rectangle subRegion)
+		{
+			if (!OperatingSystem.IsWindows())
+			{
+				Console.WriteLine("CopyAndTrimImage is not supported on this platform!");
+				return;
+			}
+
+			Console.WriteLine("Generating sub map image.\n  Source: " + sourceImagePath + "\n  Dest  : " + destImagePath);
+
+			using (FileStream fileStream = new FileStream(sourceImagePath, FileMode.Open, FileAccess.Read))
+			{
+				Console.WriteLine("Reading source image");
+				using (var image = Image.FromStream(fileStream))
+				{
+					Rectangle scaledSubRegion = subRegion;
+					float subRegionScaleW = image.Width / imageWidth;
+					float subRegionScaleH = image.Height / imageHeight;
+					scaledSubRegion.X = (int)((float)scaledSubRegion.X * subRegionScaleW);
+					scaledSubRegion.Width = (int)((float)scaledSubRegion.Width * subRegionScaleW);
+					scaledSubRegion.Y = (int)((float)scaledSubRegion.Y * subRegionScaleH);
+					scaledSubRegion.Height = (int)((float)scaledSubRegion.Height * subRegionScaleH);
+
+					if (scaledSubRegion.Width > 0 && scaledSubRegion.Height > 0)
+					{
+						// Trimming
+						Console.WriteLine("Generating Cropped Image");
+						using (Bitmap croppedBitmap = new Bitmap(scaledSubRegion.Width, scaledSubRegion.Height))
+						{
+							using (Graphics graphics = Graphics.FromImage(croppedBitmap))
+							{
+								graphics.DrawImage(image, new Rectangle(0, 0, scaledSubRegion.Width, scaledSubRegion.Height),
+												   scaledSubRegion, GraphicsUnit.Pixel);
+
+								using (var thumbnail = croppedBitmap.GetThumbnailImage(destWidth, destHeight, null, IntPtr.Zero))
+								{
+									using (var destImageStream = new FileStream(destImagePath, FileMode.OpenOrCreate, FileAccess.Write))
+									{
+										Console.WriteLine("Saving trimmed map image: " + destImagePath);
+										thumbnail.Save(destImageStream, ImageFormat.Jpeg);
+									}
+								}
+							}
+						}
+					}
+					else 
+					{
+						Console.WriteLine("Failed to scale image correctly!");
+					}
+				}
+			}
+		}
+
 		public override HtmlSection WriteSummaryData(bool bWriteHtml, CsvStats csvStats, CsvStats csvStatsUnstripped, bool bWriteSummaryCsv, SummaryTableRowData rowData, string htmlFileName)
 		{
 			HtmlSection htmlSection = null;
@@ -187,146 +308,217 @@ namespace PerfSummaries
 					}
 				}
 
-				// Check if the file exists in the output directory
-				htmlSection.WriteLine("<svg version='1.1' xmlns='http://www.w3.org/2000/svg' xmlns:xlink='http://www.w3.org/1999/xlink' width='" + imageWidth + "' height='" + imageHeight + "'>");
-				htmlSection.WriteLine("<image href='" + destImageFilename + "' width='" + imageWidth + "' height='" + imageHeight + "' />");
+				Rectangle imageRegion = new Rectangle(0, 0, (int)imageWidth, (int)imageHeight);
 
-				// Draw the overlays
-				foreach (MapOverlay overlay in overlays)
+				GenerateImageSection(csvStats, htmlSection, imageRegion, destImageFilename);
+
+				if (createCroppedImage)
 				{
-					StatSamples xStat = csvStats.GetStat(overlay.positionStatNames[0]);
-					StatSamples yStat = csvStats.GetStat(overlay.positionStatNames[1]);
+					// Generate Sub Image Location, Padded to give a small border.
+					SVGPosition minPos = toSVGLocationClamped(eventDimensions.minX, eventDimensions.minY, imageRegion);
+					SVGPosition maxPos = toSVGLocationClamped(eventDimensions.maxX, eventDimensions.maxY, imageRegion);
 
-					if (xStat == null || yStat == null)
+					if (minPos.X != maxPos.X && minPos.Y != maxPos.Y)
 					{
-						continue;
-					}
+						// Center the map.
+						int SubRegionWidth = (int)(maxPos.X - minPos.X);
+						int SubRegionHeight = (int)(maxPos.Y - minPos.Y);
+						int CenterX = minPos.X + (SubRegionWidth / 2);
+						int CenterY = minPos.Y + (SubRegionHeight / 2);
+						int UniformDim = Math.Max(SubRegionWidth, SubRegionHeight);
+						Rectangle subRegion = new Rectangle(CenterX - (UniformDim / 2), CenterY - (UniformDim / 2), UniformDim, UniformDim);
 
-					// If a startevent is specified, update the start frame
-					int startFrame = 0;
-					if (overlay.startEvent != null)
-					{
-						foreach (CsvEvent ev in csvStats.Events)
+						string croppedImageName = "crop_" + destImageFilename;
+						string outputTrimMapFilename = System.IO.Path.Combine(outputDirectory, croppedImageName);
+
+						// Skip the copy if the output file already exists
+						if (!File.Exists(outputTrimMapFilename))
 						{
-							if (CsvStats.DoesSearchStringMatch(ev.Name, overlay.startEvent.name))
+							if (File.Exists(sourceImagePath))
 							{
-								startFrame = ev.Frame;
-								break;
+								CopyAndTrimImage(sourceImagePath, outputTrimMapFilename, (int)imageWidth, (int)imageHeight, subRegion);
 							}
-						}
-					}
-
-					// Make a mapping from frame to map indices
-					List<KeyValuePair<int, MapOverlayEvent>> frameEvents = new List<KeyValuePair<int, MapOverlayEvent>>();
-					foreach (MapOverlayEvent mapEvent in overlay.events)
-					{
-						foreach (CsvEvent ev in csvStats.Events)
-						{
-							if (CsvStats.DoesSearchStringMatch(ev.Name, mapEvent.name))
+							else
 							{
-								frameEvents.Add(new KeyValuePair<int, MapOverlayEvent>(ev.Frame, mapEvent));
+								Console.WriteLine("[Warning] Can't find source map image: " + sourceImagePath);
 							}
 						}
+
+						GenerateImageSection(csvStats, htmlSection, subRegion, croppedImageName);
 					}
-					frameEvents.Sort((pair0, pair1) => pair0.Key.CompareTo(pair1.Key));
-					int eventIndex = 0;
-
-					// Draw the lines
-					string currentLineColor = overlay.lineColor;
-					string lineStartTemplate = "<polyline style='fill:none;stroke-width:1.3;stroke:{LINECOLOUR}' points='";
-					htmlSection.Write(lineStartTemplate.Replace("{LINECOLOUR}", currentLineColor));
-					float adjustedLineSplitDistanceThreshold = lineSplitDistanceThreshold * framesPerLineSegment;
-					float oldx = 0;
-					float oldy = 0;
-					int lastFrameIndex = -1;
-					for (int i = startFrame; i < xStat.samples.Count; i += framesPerLineSegment)
-					{
-						float x = xStat.samples[i];
-						float y = yStat.samples[i];
-						string lineCoordsStr = toSvgX(x, y) + "," + toSvgY(x, y) + " ";
+				}
+			}
 
-						// Figure out which event we're up to so we can do color changes
-						bool restartLineStrip = false;
-						while (eventIndex < frameEvents.Count && lastFrameIndex < frameEvents[eventIndex].Key && i >= frameEvents[eventIndex].Key)
+			// Output row data
+			if (rowData != null)
+			{
+			}
+			return htmlSection;
+		}
+
+		private void GenerateImageSection(CsvStats csvStats, HtmlSection htmlSection, Rectangle imageRegion, string imageName)
+		{
+			// Check if the file exists in the output directory
+			htmlSection.WriteLine("<svg version='1.1' xmlns='http://www.w3.org/2000/svg' xmlns:xlink='http://www.w3.org/1999/xlink' width='" + imageWidth + "' height='" + imageHeight + "'>");
+			htmlSection.WriteLine("<image href='" + imageName + "' width='" + imageWidth + "' height='" + imageHeight + "' />");
+
+			// Draw the overlays
+			foreach (MapOverlay overlay in overlays)
+			{
+				StatSamples xStat = csvStats.GetStat(overlay.positionStatNames[0]);
+				StatSamples yStat = csvStats.GetStat(overlay.positionStatNames[1]);
+
+				if (xStat == null || yStat == null)
+				{
+					continue;
+				}
+
+				// If a startevent is specified, update the start frame
+				int startFrame = 0;
+				if (overlay.startEvents.Count != 0)
+				{
+					foreach (CsvEvent ev in csvStats.Events)
+					{
+						foreach (MapOverlayEvent startEv in overlay.startEvents)
 						{
-							MapOverlayEvent mapEvent = frameEvents[eventIndex].Value;
-							string newLineColor = mapEvent.lineColor != null ? mapEvent.lineColor : overlay.lineColor;
-							// If we changed color, restart the line strip
-							if (newLineColor != currentLineColor)
+							if (CsvStats.DoesSearchStringMatch(ev.Name, startEv.name))
 							{
-								currentLineColor = newLineColor;
-								restartLineStrip = true;
+								startFrame = ev.Frame;
+								break;
 							}
-							eventIndex++;
 						}
 
-						// If the distance between this point and the last is over the threshold, restart the line strip
-						float maxManhattanDist = Math.Max(Math.Abs(x - oldx), Math.Abs(y - oldy));
-						if (maxManhattanDist > adjustedLineSplitDistanceThreshold)
+						if (startFrame != 0)
 						{
-							restartLineStrip = true;
+							break;
 						}
-						else
+					}
+				}
+
+				// Make a mapping from frame to map indices
+				List<KeyValuePair<int, MapOverlayEvent>> frameEvents = new List<KeyValuePair<int, MapOverlayEvent>>();
+				foreach (MapOverlayEvent mapEvent in overlay.events)
+				{
+					foreach (CsvEvent ev in csvStats.Events)
+					{
+						if (CsvStats.DoesSearchStringMatch(ev.Name, mapEvent.name))
 						{
-							htmlSection.Write(lineCoordsStr);
+							frameEvents.Add(new KeyValuePair<int, MapOverlayEvent>(ev.Frame, mapEvent));
 						}
+					}
+				}
+				frameEvents.Sort((pair0, pair1) => pair0.Key.CompareTo(pair1.Key));
+				int eventIndex = 0;
 
-						if (restartLineStrip)
+				// Draw the lines
+				string currentLineColor = overlay.lineColor;
+				string lineStartTemplate = "<polyline style='fill:none;stroke-width:1.3;stroke:{LINECOLOUR}' points='";
+				htmlSection.Write(lineStartTemplate.Replace("{LINECOLOUR}", currentLineColor));
+				float adjustedLineSplitDistanceThreshold = lineSplitDistanceThreshold * framesPerLineSegment;
+				float oldx = 0;
+				float oldy = 0;
+				int lastFrameIndex = -1;
+				for (int i = startFrame; i < xStat.samples.Count; i += framesPerLineSegment)
+				{
+					float x = xStat.samples[i];
+					float y = yStat.samples[i];
+					SVGPosition svgPos = toSVGLocation(x, y, imageRegion);
+
+					string lineCoordsStr = svgPos.X + "," + svgPos.Y + " ";
+
+					eventDimensions.minX = Math.Min(eventDimensions.minX, x);
+					eventDimensions.maxX = Math.Max(eventDimensions.maxX, x);
+					eventDimensions.minY = Math.Min(eventDimensions.minY, y);
+					eventDimensions.maxY = Math.Max(eventDimensions.maxY, y);
+
+					// Figure out which event we're up to so we can do color changes
+					bool restartLineStrip = false;
+					while (eventIndex < frameEvents.Count && lastFrameIndex < frameEvents[eventIndex].Key && i >= frameEvents[eventIndex].Key)
+					{
+						MapOverlayEvent mapEvent = frameEvents[eventIndex].Value;
+						string newLineColor = mapEvent.lineColor != null ? mapEvent.lineColor : overlay.lineColor;
+						// If we changed color, restart the line strip
+						if (newLineColor != currentLineColor)
 						{
-							htmlSection.WriteLine("'/>");
-							htmlSection.Write(lineStartTemplate.Replace("{LINECOLOUR}", currentLineColor));
-							htmlSection.Write(lineCoordsStr);
+							currentLineColor = newLineColor;
+							restartLineStrip = true;
 						}
-						oldx = x;
-						oldy = y;
-						lastFrameIndex = i;
+						eventIndex++;
+					}
+
+					// If the distance between this point and the last is over the threshold, restart the line strip
+					float maxManhattanDist = Math.Max(Math.Abs(x - oldx), Math.Abs(y - oldy));
+					if (maxManhattanDist > adjustedLineSplitDistanceThreshold)
+					{
+						restartLineStrip = true;
+					}
+					else
+					{
+						htmlSection.Write(lineCoordsStr);
+					}
+
+					if (restartLineStrip)
+					{
+						htmlSection.WriteLine("'/>");
+						htmlSection.Write(lineStartTemplate.Replace("{LINECOLOUR}", currentLineColor));
+						htmlSection.Write(lineCoordsStr);
 					}
-					htmlSection.WriteLine("'/>");
+					oldx = x;
+					oldy = y;
+					lastFrameIndex = i;
+				}
+				htmlSection.WriteLine("'/>");
 
-					// Plot the events 
-					float circleRadius = 3;
-					string eventColourString = "#ffffff";
-					foreach (MapOverlayEvent mapEvent in overlay.events)
+				// Plot the events 
+				float circleRadius = 3;
+				string eventColourString = "#ffffff";
+				foreach (MapOverlayEvent mapEvent in overlay.events)
+				{
+					foreach (CsvEvent ev in csvStats.Events)
 					{
-						foreach (CsvEvent ev in csvStats.Events)
+						if (CsvStats.DoesSearchStringMatch(ev.Name, mapEvent.name))
 						{
-							if (CsvStats.DoesSearchStringMatch(ev.Name, mapEvent.name))
-							{
-								string eventText = mapEvent.shortName != null ? mapEvent.shortName : ev.Name;
-								float x = xStat.samples[ev.Frame];
-								float y = yStat.samples[ev.Frame];
-								int svgX = toSvgX(x, y);
-								int svgY = toSvgY(x, y);
-								htmlSection.Write("<circle cx='" + svgX + "' cy='" + svgY + "' r='" + circleRadius + "' fill='" + eventColourString + "' fill-opacity='1.0'/>");
-								htmlSection.WriteLine("<text x='" + (svgX + 5) + "' y='" + svgY + "' text-anchor='left' style='font-family: Verdana;fill: #ffffff; font-size: " + 9 + "px;'>" + eventText + "</text>");
-							}
+							string eventText = mapEvent.shortName != null ? mapEvent.shortName : ev.Name;
+							float x = xStat.samples[ev.Frame];
+							float y = yStat.samples[ev.Frame];
+							SVGPosition svgPos = toSVGLocation(x, y, imageRegion);
+							htmlSection.Write("<circle cx='" + svgPos.X + "' cy='" + svgPos.Y + "' r='" + circleRadius + "' fill='" + eventColourString + "' fill-opacity='1.0'/>");
+							htmlSection.WriteLine("<text x='" + (svgPos.X + 5) + "' y='" + svgPos.Y + "' text-anchor='left' style='font-family: Verdana;fill: #ffffff; font-size: " + 9 + "px;'>" + eventText + "</text>");
 						}
 					}
 				}
-
-				//htmlSection.WriteLine("<text x='50%' y='" + (imageHeight * 0.05) + "' text-anchor='middle' style='font-family: Verdana;fill: #FFFFFF; stroke: #C0C0C0;  font-size: " + 20 + "px;'>" + title + "</text>");
-				htmlSection.WriteLine("</svg>");
 			}
 
-			// Output row data
-			if (rowData != null)
-			{
-			}
-			return htmlSection;
+			//htmlSection.WriteLine("<text x='50%' y='" + (imageHeight * 0.05) + "' text-anchor='middle' style='font-family: Verdana;fill: #FFFFFF; stroke: #C0C0C0;  font-size: " + 20 + "px;'>" + title + "</text>");
+			htmlSection.WriteLine("</svg>");
 		}
+
 		public override void PostInit(ReportTypeInfo reportTypeInfo, CsvStats csvStats)
 		{
 		}
+
+		class MapEventDimensions
+		{
+			public float minX = float.MaxValue;
+			public float maxX = float.MinValue;
+			public float minY = float.MaxValue;
+			public float maxY = float.MinValue;
+		}
+
 		string title;
 		string sourceImagePath;
 		float[] transform = new float[6];
+		float[] fullTransform = null; // Default to null while we transition
+		Matrix4x4 transformMatrix = Matrix4x4.Identity;
 		string destImageFilename;
 		float imageWidth;
 		float imageHeight;
 		float lineSplitDistanceThreshold;
 		int framesPerLineSegment;
+		bool createCroppedImage = false;
 
 		List<MapOverlay> overlays = new List<MapOverlay>();
+		MapEventDimensions eventDimensions = new MapEventDimensions();
 	};
 
 }
\ No newline at end of file
