Changed the RemoveSpacesInline on the FString value type to a TrimStartAndEnd in CSVToTelemetry.The former operation would strip out required spaces in legitimate text values.
Renamed the index from i to Column as this is more representative.

diff --git a/Engine/Source/Programs/CSVTools/CSVToTelemetry/Private/CSVToTelemetry.cpp b/Engine/Source/Programs/CSVTools/CSVToTelemetry/Private/CSVToTelemetry.cpp
--- a/Engine/Source/Programs/CSVTools/CSVToTelemetry/Private/CSVToTelemetry.cpp
+++ b/Engine/Source/Programs/CSVTools/CSVToTelemetry/Private/CSVToTelemetry.cpp
@@ -196,17 +196,17 @@ static int GenerateTelemetryFromCSVFile(const FString& FilePath)
 						TArray<FAnalyticsEventAttribute> Attributes = AdditionalAttributes;
 
 						// Add the Key/Value pair to the attribute list
-						for (int32 i = 0; i < ValueArray.Num(); ++i)
+						for (int32 Column = 0; Column < ValueArray.Num(); ++Column)
 						{
-							ValueArray[i].RemoveSpacesInline();
-
+							FString Value = ValueArray[Column].TrimStartAndEnd();
+						
 							// Normalize boolean strings
 							if (NormalizeBooleanStrings)
 							{
-								TryNormalizeBooleanString(ValueArray[i]);
+								TryNormalizeBooleanString(Value);
 							}
 
-							Attributes.Emplace(KeyArray[i], ValueArray[i]);
+							Attributes.Emplace(KeyArray[Column], Value);
 						}
 
 						// Keep track of the total data we're uploading
