PerfReportTool 4.261.0 - Improved timing region summary
- Make region times table collapsible (collapsed by default), update heading name
- Filter out regions < 3px in width
- Make track sorting stable
- Add support for beginRegion and endRegion as an alternative to beginEvent/endEvent (note: specifying both is not allowed)
- Rename startEvent to beginEvent for consistency

diff --git a/Engine/Source/Programs/CSVTools/PerfReportTool/PerfReportTool.cs b/Engine/Source/Programs/CSVTools/PerfReportTool/PerfReportTool.cs
--- a/Engine/Source/Programs/CSVTools/PerfReportTool/PerfReportTool.cs
+++ b/Engine/Source/Programs/CSVTools/PerfReportTool/PerfReportTool.cs
@@ -23,7 +23,7 @@ namespace PerfReportTool
     class Version
     {
 		// Format: Major.Minor.Bugfix
-        private static string VersionString = "4.260.0";
+        private static string VersionString = "4.261.0";
 
         public static string Get() { return VersionString; }
     };
diff --git a/Engine/Source/Programs/CSVTools/PerfReportTool/Summaries/TimingRegionSummary.cs b/Engine/Source/Programs/CSVTools/PerfReportTool/Summaries/TimingRegionSummary.cs
--- a/Engine/Source/Programs/CSVTools/PerfReportTool/Summaries/TimingRegionSummary.cs
+++ b/Engine/Source/Programs/CSVTools/PerfReportTool/Summaries/TimingRegionSummary.cs
@@ -34,8 +34,30 @@ namespace PerfSummaries
 				throw new Exception("<stats> element is not supported");
 			}
 			title = element.GetRequiredAttribute<string>(vars, "title");
-			startEventName = element.GetSafeAttribute<string>(vars, "startEvent");
+			beginEventName = element.GetSafeAttribute<string>(vars, "beginEvent");
+			string beginRegionName = element.GetSafeAttribute<string>(vars, "beginRegion");
+			if (beginRegionName != null)
+			{
+				if (beginEventName != null)
+				{
+					throw new Exception("Timing region summary has both beginEvent and beginRegion defined!");
+				}
+				beginEventName = $"Region: {beginRegionName} - Begin";
+			}
+
 			endEventName = element.GetSafeAttribute<string>(vars, "endEvent");
+			string endRegionName = element.GetSafeAttribute<string>(vars, "endRegion");
+			if (endRegionName != null)
+			{
+				if ( endEventName != null)
+				{
+					throw new Exception("Timing region summary has both endEvent and endRegion defined!");
+				}
+				endEventName = $"Region: {endRegionName} - End";
+			}
+
+			bShowTimingTable = element.GetSafeAttribute<bool>(vars, "showTimingTable", true);
+			
 			metricPrefix = element.GetSafeAttribute<string>(vars, "metricPrefix");
 			foreach (XElement regionFilterXml in element.Elements("regionFilter"))
 			{
@@ -280,15 +302,20 @@ namespace PerfSummaries
 			return newRegionEvents;
 		}
 
-		class RegionEventTrack
+		class RegionEventTrack : IComparable<RegionEventTrack>
 		{
-			public RegionEventTrack(string inCategory, double inStartTime, double inEndTime, RegionFilterInfo inFilterInfo)
+			public RegionEventTrack(string inCategory, double inStartTime, double inEndTime, RegionFilterInfo inFilterInfo, double inTrackWidth, double inTrackHeight)
 			{
 				category = inCategory;
 				startTime = inStartTime;
 				endTime = inEndTime;
 				filterInfo = inFilterInfo;
 				filterIndex = filterInfo.index;
+				trackWidth = inTrackWidth;
+				trackHeight = inTrackHeight;
+				double widthThreshold = 3;
+				double trackDuration = (endTime - startTime);
+				regionDurationThreshold = ( widthThreshold / trackWidth ) * trackDuration;
 			}
 			public bool TryAddEvent(RegionEvent beginEvent)
 			{
@@ -296,12 +323,20 @@ namespace PerfSummaries
 				{
 					return false;
 				}
+				if (beginEvent.GetPairDuration() < regionDurationThreshold)
+				{
+					return false;
+				}
+				if (firstEventTime == double.MaxValue)
+				{
+					firstEventTime = beginEvent.timestamp;
+				}
 				lastEventTime = (beginEvent.otherEvent != null) ? beginEvent.otherEvent.timestamp : endTime;
 				events.Add(beginEvent);
 				return true;
 			}
 
-			public void WriteToHtml(HtmlSection htmlSection, double xOffset, double yOffset, double trackWidth, double trackHeight, string backgroundColor="#ffffff")
+			public void WriteToHtml(HtmlSection htmlSection, double xOffset, double yOffset, string backgroundColor="#ffffff")
 			{
 				// Strip the hidePrefix if necessary
 				string backgroundStyle = "fill:" + backgroundColor;
@@ -315,7 +350,10 @@ namespace PerfSummaries
 					double t = beginEvent.timestamp - startTime;
 					double x = Math.Round(xOffset + trackWidth * t / trackDuration,0);
 					double width = Math.Round(trackWidth * duration / trackDuration,0);
-
+					if (width < 3)
+					{
+						continue;
+					}
 					// Truncate the display name if required
 					string displayName = beginEvent.name;
 					if (filterInfo.hidePrefix != null && beginEvent.name.StartsWith(filterInfo.hidePrefix))
@@ -352,12 +390,26 @@ namespace PerfSummaries
 				htmlSection.WriteLine($"<rect x='{(lastEventX+1)}' y='{yOffset}' width='{(trackWidth-(lastEventX-xOffset)-1)}' height='{trackHeight}' style='{backgroundStyle}'/>");
 			}
 
+			public int CompareTo(RegionEventTrack other)
+			{
+				int filterIndexCompare = filterIndex.CompareTo(other.filterIndex);
+				if (filterIndexCompare != 0)
+				{
+					return filterIndexCompare;
+				}
+				return firstEventTime.CompareTo(other.firstEventTime);
+			}
+
 			public string category;
 			public int filterIndex;
 			public double lastEventTime;
+			public double firstEventTime = double.MaxValue;
 			public List<RegionEvent> events = new List<RegionEvent>();
 			public double startTime;
 			public double endTime;
+			public double trackWidth;
+			public double trackHeight;
+			public double regionDurationThreshold;
 			public RegionFilterInfo filterInfo;
 		}
 
@@ -371,8 +423,8 @@ namespace PerfSummaries
 				throw new Exception("TimingRegion summary requires a frametime stat");
 			}
 
-			CsvEvent startEvent = (startEventName != null) ? csvStatsUnstripped.FindEvent(startEventName) : null;
-			double startTime = (startEvent != null) ? startEvent.Timestamp : 0.0;
+			CsvEvent beginEvent = (beginEventName != null) ? csvStatsUnstripped.FindEvent(beginEventName) : null;
+			double startTime = (beginEvent != null) ? beginEvent.Timestamp : 0.0;
 
 			CsvEvent endEvent = (endEventName != null) ? csvStatsUnstripped.FindEvent(endEventName, startTime) : null;
 			double endTime = (endEvent != null) ? endEvent.Timestamp : frameTimeStat.total / 1000.0;
@@ -386,47 +438,47 @@ namespace PerfSummaries
 				return null;
 			}
 
-			// Assign events to tracks
-			List<RegionEventTrack> tracks = new List<RegionEventTrack>();
-			foreach (RegionEvent ev in regionEvents)
+			// Output HTML
+			string prevCategory = null;
+			if (bWriteHtml)
 			{
-				if (ev.bBegin && ev.otherEvent != null)
+				htmlSection = new HtmlSection(title, bStartCollapsed);
+
+				double imageWidth = 1400;
+				double headerWidth = 150;
+				double trackWidth = imageWidth - headerWidth;
+
+				// Assign events to tracks
+				List<RegionEventTrack> tracks = new List<RegionEventTrack>();
+				foreach (RegionEvent ev in regionEvents)
 				{
-					// Add the event to an available track
-					bool bFoundTrack = false;
-					foreach (RegionEventTrack track in tracks)
+					if (ev.bBegin && ev.otherEvent != null)
 					{
-						if (track.TryAddEvent(ev))
+						// Add the event to an available track
+						bool bFoundTrack = false;
+						foreach (RegionEventTrack track in tracks)
 						{
-							bFoundTrack = true;
-							break;
+							if (track.TryAddEvent(ev))
+							{
+								bFoundTrack = true;
+								break;
+							}
 						}
-					}
-					if ( !bFoundTrack )
-					{
-						RegionEventTrack newTrack = new RegionEventTrack(ev.category, startTime, endTime, ev.filterInfo);
-						if ( newTrack.TryAddEvent(ev) )
+						if (!bFoundTrack)
 						{
-							tracks.Add(newTrack);
+							RegionEventTrack newTrack = new RegionEventTrack(ev.category, startTime, endTime, ev.filterInfo, trackWidth, 20);
+							if (newTrack.TryAddEvent(ev))
+							{
+								tracks.Add(newTrack);
+							}
 						}
 					}
 				}
-			}
 
-			// Sort tracks by filter index
-			tracks.Sort((a, b) => a.filterIndex.CompareTo(b.filterIndex));
+				// Sort tracks
+				tracks.Sort();
 
-			// Output HTML
-			string prevCategory = null;
-			if (bWriteHtml)
-			{
-				htmlSection = new HtmlSection(title, bStartCollapsed);
-
-				double imageWidth = 1400;
 				double imageHeight = tracks.Count * 22;
-				double headerWidth = 150;
-				double trackWidth = imageWidth - headerWidth;
-
 				string lineStyle = "stroke:#808080;stroke-width:2";
 				string categoryTextStyle = "font-family:verdana;font-size:12px";
 
@@ -450,7 +502,7 @@ namespace PerfSummaries
 						prevCategory = track.category;
 						htmlSection.WriteLine($"<text x='{0}' y='{yOffset+14}' style='{categoryTextStyle}'>{track.category}</text>");
 					}
-					track.WriteToHtml(htmlSection, headerWidth, yOffset, trackWidth, 20, trackIndex%2==0 ? "#e0e0e0" : "#f0f0f0");
+					track.WriteToHtml(htmlSection, headerWidth, yOffset, trackIndex%2==0 ? "#e0e0e0" : "#f0f0f0");
 					yOffset += 20;
 					trackIndex++;
 				}
@@ -481,19 +533,26 @@ namespace PerfSummaries
 					rowData.Add(new SummaryTableElement(SummaryTableElement.Type.SummaryTableMetric, name, pair.Value, null, ""));
 				}
 			}
-			if (htmlSection != null)
+			if (htmlSection != null && bShowTimingTable)
 			{
-				htmlSection.WriteLine("<h3>Total times</h3>");
+				htmlSection.WriteLine("<h3 class='collapsibleHeading collapsed'>Region durations</h3>");
+				htmlSection.WriteLine("<div class='collapsibleSection collapsed'>");
+				htmlSection.WriteLine("<div class='collapsibleSectionInner'>");
 				htmlSection.WriteLine("<table border='0' style='width:1000'>");
-				htmlSection.WriteLine("<tr><th>Region</th><th>Time</th></tr>");
+				htmlSection.WriteLine("<tr><th>Region</th><th>Time (s)</th></tr>");
 				List<KeyValuePair<string, double>> regionTimingsSorted = new List<KeyValuePair<string, double>>(regionTotalTimes);
 				regionTimingsSorted.Sort((a, b) => a.Key.CompareTo(b.Key));
 
 				foreach (KeyValuePair<string, double> pair in regionTimingsSorted)
 				{
-					htmlSection.WriteLine($"<tr><td>{pair.Key}</td><td>{String.Format("{0:0.00}", pair.Value)}</td></tr>");
+					if (pair.Value>=0.01)
+					{
+						htmlSection.WriteLine($"<tr><td>{pair.Key}</td><td>{String.Format("{0:0.00}", pair.Value)}</td></tr>");
+					}
 				}
 				htmlSection.WriteLine("</table>");
+				htmlSection.WriteLine("</div>");
+				htmlSection.WriteLine("</div>");
 			}
 			return htmlSection;
 		}
@@ -502,9 +561,10 @@ namespace PerfSummaries
 		}
 
 		string title;
-		string startEventName;
+		string beginEventName;
 		string endEventName;
 		string metricPrefix;
+		bool bShowTimingTable;
 		List<RegionFilterInfo> regionFilters = new List<RegionFilterInfo>();
 	};
 
