PerfReportTool: Add a trackDisplayThreshold attribute to the timingregion summary region filter which allows us to filter out short regions based on display width. Default to 1px

diff --git a/Engine/Source/Programs/CSVTools/PerfReportTool/Summaries/TimingRegionSummary.cs b/Engine/Source/Programs/CSVTools/PerfReportTool/Summaries/TimingRegionSummary.cs
--- a/Engine/Source/Programs/CSVTools/PerfReportTool/Summaries/TimingRegionSummary.cs
+++ b/Engine/Source/Programs/CSVTools/PerfReportTool/Summaries/TimingRegionSummary.cs
@@ -19,12 +19,14 @@ namespace PerfSummaries
 				pattern = element.GetRequiredAttribute<string>("pattern");
 				busyTimeMetric = element.GetSafeAttribute<string>("busyTimeMetric");
 				hidePrefix = element.GetSafeAttribute<string>("hidePrefix");
+				trackDisplayThreshold = element.GetSafeAttribute<int>("trackDisplayThreshold", 1);
 				index = indexIn;
 			}
 			public string pattern;
 			public string busyTimeMetric = null;
 			public string hidePrefix = null;
 			public int index = 0;
+			public int trackDisplayThreshold = 0;
 		}
 		public TimingRegionSummary(XElement element, XmlVariableMappings vars, string baseXmlDirectory)
 		{
@@ -313,9 +315,8 @@ namespace PerfSummaries
 				filterIndex = filterInfo.index;
 				trackWidth = inTrackWidth;
 				trackHeight = inTrackHeight;
-				double widthThreshold = 3;
 				double trackDuration = (endTime - startTime);
-				regionDurationThreshold = ( widthThreshold / trackWidth ) * trackDuration;
+				regionDurationThreshold = (inFilterInfo.trackDisplayThreshold / trackWidth ) * trackDuration;
 			}
 			public bool TryAddEvent(RegionEvent beginEvent)
 			{
