Tweaks and fixes to CSVToTelemetry
Fixed an issue with addiotnal spaces being added to Key entries causing invalid event data to be sent
Fixed the MaxRows logic so that the Keys are not counted
Added validation option which parses the CSV file but does not upload any data.

#Synced-CL 48281812

diff --git a/Engine/Source/Programs/CSVTools/CSVToTelemetry/Private/CSVToTelemetry.cpp b/Engine/Source/Programs/CSVTools/CSVToTelemetry/Private/CSVToTelemetry.cpp
--- a/Engine/Source/Programs/CSVTools/CSVToTelemetry/Private/CSVToTelemetry.cpp
+++ b/Engine/Source/Programs/CSVTools/CSVToTelemetry/Private/CSVToTelemetry.cpp
@@ -45,6 +45,14 @@ static int GenerateTelemetryFromCSVFile(const FString& FilePath)
 			return Result;
 		}
 
+		bool Validate = false;
+
+		if (FParse::Param(FCommandLine::Get(), TEXT("validate")))
+		{
+			Validate = true;
+			UE_LOG(LogCSVToTelemetry, Error, TEXT("Validating CSV file %s, events will not be sent"), *FilePath);
+		}
+
 		STUDIO_TELEMETRY_SESSION_SCOPE
 
 		// Start the telemetry session and record each row as a single telemetry event 
@@ -62,7 +70,7 @@ static int GenerateTelemetryFromCSVFile(const FString& FilePath)
 				if ( Rows.Num()==0 )
 				{
 					// The CSV file failed to parse
-					UE_LOG(LogCSVToTelemetry, Warning, TEXT("Failed to import rows from file %s"), *FilePath);;
+					UE_LOG(LogCSVToTelemetry, Warning, TEXT("Failed to import rows from file %s"), *FilePath);
 					return -1;
 				}
 				else
@@ -74,7 +82,7 @@ static int GenerateTelemetryFromCSVFile(const FString& FilePath)
 				uint32 TotalEventsSent = 0;
 				uint32 KeySize = 0;
 				uint32 ValueSize = 0;
-				uint32 MaxRows = -1;
+				uint32 MaxRows = Rows.Num();
 				TArray<FString> KeyArray;
 				TArray<FAnalyticsEventAttribute> AdditionalAttributes;
 
@@ -144,12 +152,15 @@ static int GenerateTelemetryFromCSVFile(const FString& FilePath)
 
 					for (const TCHAR* Column : Row)
 					{
-						KeyArray.Push(Column);
+						FString ColumnString(Column);
+						KeyArray.Push(ColumnString.TrimStartAndEnd());
 					}
 
 					KeySize = Row.NumBytes();
 
+					// Increment the RowIndex and the MaxRows as we have consumed the fist row for the keys  
 					RowIndex++;
+					MaxRows++;
 				}
 
 				// Validate that we have a valid set of keys. These are either specified on the first row of the csv file or explicitly via the -columns= commandline
@@ -213,8 +224,11 @@ static int GenerateTelemetryFromCSVFile(const FString& FilePath)
 						TotalUploadSize += KeySize + ValueSize;
 
 						// Now we have a complete event so send it
-						FStudioTelemetry::Get().RecordEvent(EventName, Attributes);
-
+						if (Validate == false)
+						{
+							FStudioTelemetry::Get().RecordEvent(EventName, Attributes);
+						}
+						
 						// Keep track of the event count
 						TotalEventsSent++;
 					}
