PerfReportTool 4.261.1: Allow underscores in derived metadata names
#Synced-CL 0

diff --git a/Engine/Source/Programs/CSVTools/PerfReportTool/PerfReportTool.cs b/Engine/Source/Programs/CSVTools/PerfReportTool/PerfReportTool.cs
--- a/Engine/Source/Programs/CSVTools/PerfReportTool/PerfReportTool.cs
+++ b/Engine/Source/Programs/CSVTools/PerfReportTool/PerfReportTool.cs
@@ -23,7 +23,7 @@ namespace PerfReportTool
     class Version
     {
 		// Format: Major.Minor.Bugfix
-        private static string VersionString = "4.261.0";
+        private static string VersionString = "4.261.1";
 
         public static string Get() { return VersionString; }
     };
diff --git a/Engine/Source/Programs/CSVTools/PerfReportTool/ReportXml.cs b/Engine/Source/Programs/CSVTools/PerfReportTool/ReportXml.cs
--- a/Engine/Source/Programs/CSVTools/PerfReportTool/ReportXml.cs
+++ b/Engine/Source/Programs/CSVTools/PerfReportTool/ReportXml.cs
@@ -701,7 +701,7 @@ namespace PerfReportTool
 			if (bValidate)
 			{
 				// Check for legal characters in the name
-				if (!Name.All(x => char.IsLetterOrDigit(x) || x == '.'))
+				if (!Name.All(x => char.IsLetterOrDigit(x) || x == '.' || x=='_'))
 				{
 					throw new Exception("Invalid variable name: " + Name);
 				}
