PerfReportTool 4.262.0: Metadata query fixes
-  Metadata not found in an expression gets a value of "NULL" rather than the expression always evaluating to false. This is more intuitive and behaves as you'd expect with equality tests, e.g metadata!=0 would evalute to true if the metadata doesn't exist. You can also do metadata!=NULL to query non-valid metadata.
- Fix Negation not working consistently

#tests tested with PRT unit tests
#Synced-CL 0

diff --git a/Engine/Source/Programs/CSVTools/CsvStats/MetadataQueryBuilder.cs b/Engine/Source/Programs/CSVTools/CsvStats/MetadataQueryBuilder.cs
--- a/Engine/Source/Programs/CSVTools/CsvStats/MetadataQueryBuilder.cs
+++ b/Engine/Source/Programs/CSVTools/CsvStats/MetadataQueryBuilder.cs
@@ -1,6 +1,7 @@
 using System;
-using System.Diagnostics;
 using System.Collections.Generic;
+using System.Diagnostics;
+using System.Globalization;
 
 namespace CSVStats
 {
@@ -37,19 +38,32 @@ namespace CSVStats
 	{
 		public override bool Evaluate(CsvMetadata metadata)
 		{
-			if (!metadata.Values.ContainsKey(Key))
-			{
-				return false;
-			}
+
 			bool matches = false;
 			bool negate = bNegate;
-			string metaDataValueStr = metadata.Values[Key].ToLower();
 
-			// Attempt to parse both values as numeric and determine if this is a numeric comparison
-			double doubleVal = double.MaxValue;
-			bool bNumericComparison = Double.TryParse(Value, out doubleVal);
+			double doubleVal;
+			if (!Double.TryParse(Value, NumberStyles.Float, CultureInfo.InvariantCulture, out doubleVal))
+			{
+				doubleVal = double.MaxValue;
+			}
+
+			string metaDataValueStr = "NULL";
 			double doubleMetadataVal = double.MaxValue;
-			bNumericComparison &= Double.TryParse(metaDataValueStr, out doubleMetadataVal);
+			bool bNumericComparison = false;
+			if (metadata.Values.TryGetValue(Key, out string metaDataValueCased))
+			{
+				metaDataValueStr = metaDataValueCased.ToLower();
+				// Attempt to parse both values as numeric and determine if this is a numeric comparison
+				if (!Double.TryParse(metaDataValueStr, NumberStyles.Float, CultureInfo.InvariantCulture, out doubleMetadataVal))
+				{
+					doubleMetadataVal = double.MaxValue; 
+				}
+				if(doubleVal != double.MaxValue && doubleMetadataVal != double.MaxValue)
+				{
+					bNumericComparison = true;
+				}
+			}
 
 			if (Op == ComparisonOp.Equal || Op == ComparisonOp.NotEqual)
 			{
@@ -58,19 +72,37 @@ namespace CSVStats
 				{
 					matches = doubleVal == doubleMetadataVal;
 				}
-				else 
+				else
 				{
 					// Do a string comparison (taking wildcards into account)
 					matches = CsvStats.DoesSearchStringMatch(metaDataValueStr, Value.ToLower());
-					if (Op == ComparisonOp.NotEqual)
-					{
-						negate = !negate;
-					}
+				}
+				if (Op == ComparisonOp.NotEqual)
+				{
+					negate = !negate;
 				}
 			}
 			else if (Op == ComparisonOp.LessThan || Op == ComparisonOp.GreaterThan || Op == ComparisonOp.LessThanOrEqual || Op == ComparisonOp.GreaterThanOrEqual)
 			{
-				if (!bNumericComparison)
+				if (bNumericComparison)
+				{
+					switch (Op)
+					{
+						case ComparisonOp.LessThan:
+							matches = doubleMetadataVal < doubleVal;
+							break;
+						case ComparisonOp.GreaterThan:
+							matches = doubleMetadataVal > doubleVal;
+							break;
+						case ComparisonOp.LessThanOrEqual:
+							matches = doubleMetadataVal <= doubleVal;
+							break;
+						case ComparisonOp.GreaterThanOrEqual:
+							matches = doubleMetadataVal >= doubleVal;
+							break;
+					}
+				}
+				else
 				{
 					if (doubleVal == double.MaxValue)
 					{
@@ -78,22 +110,10 @@ namespace CSVStats
 					}
 					if (doubleMetadataVal == double.MaxValue)
 					{
-						// Metadata value wasn't numeric, so just return
-						Console.WriteLine("Value for " + Key + " was non-numeric: " + Value);
-						return false;
+						// Metadata value wasn't numeric, so just return false. This will be the case with NULL
+						matches = false;
 					}
 				}
-				switch (Op)
-				{
-					case ComparisonOp.LessThan:
-						return doubleMetadataVal < doubleVal;
-					case ComparisonOp.GreaterThan:
-						return doubleMetadataVal > doubleVal;
-					case ComparisonOp.LessThanOrEqual:
-						return doubleMetadataVal <= doubleVal;
-					case ComparisonOp.GreaterThanOrEqual:
-						return doubleMetadataVal >= doubleVal;
-				}
 			}
 			else if (Op == ComparisonOp.Contains)
 			{
diff --git a/Engine/Source/Programs/CSVTools/PerfReportTool/PerfReportTool.cs b/Engine/Source/Programs/CSVTools/PerfReportTool/PerfReportTool.cs
--- a/Engine/Source/Programs/CSVTools/PerfReportTool/PerfReportTool.cs
+++ b/Engine/Source/Programs/CSVTools/PerfReportTool/PerfReportTool.cs
@@ -23,7 +23,7 @@ namespace PerfReportTool
     class Version
     {
 		// Format: Major.Minor.Bugfix
-        private static string VersionString = "4.261.1";
+        private static string VersionString = "4.262.0";
 
         public static string Get() { return VersionString; }
     };
