CSV tools: Default to invariant culture for in the CommandLineTool base class. Fixes issu with CsvConvert with european culture
#Synced-CL 0

diff --git a/Engine/Source/Programs/CSVTools/CsvStats/CommandLineTool.cs b/Engine/Source/Programs/CSVTools/CsvStats/CommandLineTool.cs
--- a/Engine/Source/Programs/CSVTools/CsvStats/CommandLineTool.cs
+++ b/Engine/Source/Programs/CSVTools/CsvStats/CommandLineTool.cs
@@ -236,6 +236,7 @@ namespace CSVStats
 		public CommandLineTool()
 		{
 			bMainInstance = Interlocked.Increment(ref initCounter) == 1;
+			System.Globalization.CultureInfo.DefaultThreadCurrentCulture = System.Globalization.CultureInfo.InvariantCulture;
 		}
 
 		// Accesses the commandline for the main instance. Caution: if using batch CommandLineTool instances (e.g via CsvToSvg -batchCommands), this will not return the commandline for the batch
