PerfReportTool 4.262.1: Fix a bug with region visualization causing mismatched pairs and incorrectly long regions in some cases. We were sorting events by timestamp but failed to use the sorted array.
#Synced-CL 46049276

diff --git a/Engine/Source/Programs/CSVTools/PerfReportTool/PerfReportTool.cs b/Engine/Source/Programs/CSVTools/PerfReportTool/PerfReportTool.cs
--- a/Engine/Source/Programs/CSVTools/PerfReportTool/PerfReportTool.cs
+++ b/Engine/Source/Programs/CSVTools/PerfReportTool/PerfReportTool.cs
@@ -23,7 +23,7 @@ namespace PerfReportTool
     class Version
     {
 		// Format: Major.Minor.Bugfix
-        private static string VersionString = "4.262.0";
+        private static string VersionString = "4.262.1";
 
         public static string Get() { return VersionString; }
     };
diff --git a/Engine/Source/Programs/CSVTools/PerfReportTool/Summaries/TimingRegionSummary.cs b/Engine/Source/Programs/CSVTools/PerfReportTool/Summaries/TimingRegionSummary.cs
--- a/Engine/Source/Programs/CSVTools/PerfReportTool/Summaries/TimingRegionSummary.cs
+++ b/Engine/Source/Programs/CSVTools/PerfReportTool/Summaries/TimingRegionSummary.cs
@@ -193,7 +193,7 @@ namespace PerfSummaries
 			List<CsvEvent> csvEvents = new List<CsvEvent>(csvStats.Events);
 			csvEvents.Sort((a, b) => a.Timestamp.CompareTo(b.Timestamp));
 
-			foreach (CsvEvent ev in csvStats.Events)
+			foreach (CsvEvent ev in csvEvents)
 			{
 				RegionEvent regionEvent = RegionEvent.TryCreateFromCsvEvent(ev);
 				if (regionEvent != null && TryGetRegionFilterInfo(regionEvent, out regionEvent.filterInfo))
