PerfReportTool 4.262.2: Fix a bug with region visualization causing mismatched pairs and incorrectly long regions in some cases (pt2). Do a stable sort (using Linq OrderBy), List.sort is not stable!

diff --git a/Engine/Source/Programs/CSVTools/PerfReportTool/Summaries/TimingRegionSummary.cs b/Engine/Source/Programs/CSVTools/PerfReportTool/Summaries/TimingRegionSummary.cs
--- a/Engine/Source/Programs/CSVTools/PerfReportTool/Summaries/TimingRegionSummary.cs
+++ b/Engine/Source/Programs/CSVTools/PerfReportTool/Summaries/TimingRegionSummary.cs
@@ -189,10 +189,8 @@ namespace PerfSummaries
 			// Find the region events
 			List<RegionEvent> allRegionEvents = new List<RegionEvent>();
 
-			// Make sure the events are sorted by timestamp
-			List<CsvEvent> csvEvents = new List<CsvEvent>(csvStats.Events);
-			csvEvents.Sort((a, b) => a.Timestamp.CompareTo(b.Timestamp));
-
+			// Make sure the events are sorted by timestamp. Note: We need a stable sort here, so we use Linq Orderby
+			List<CsvEvent> csvEvents = csvStats.Events.OrderBy(e => e.Timestamp).ToList();
 			foreach (CsvEvent ev in csvEvents)
 			{
 				RegionEvent regionEvent = RegionEvent.TryCreateFromCsvEvent(ev);
@@ -300,7 +298,7 @@ namespace PerfSummaries
 				}
 			}
 
-			newRegionEvents.Sort((a, b) => a.timestamp.CompareTo(b.timestamp));
+			newRegionEvents = newRegionEvents.OrderBy(e => e.timestamp).ToList();
 			return newRegionEvents;
 		}
 
