Updated CSVToTelemetry to latest version.
Uses built in CSV Parser instead of custom tokenized CSV parser.
Added Version to log output

#Synced-CL 49023345

diff --git a/Engine/Source/Programs/CSVTools/CSVToTelemetry/Private/CSVToTelemetry.cpp b/Engine/Source/Programs/CSVTools/CSVToTelemetry/Private/CSVToTelemetry.cpp
--- a/Engine/Source/Programs/CSVTools/CSVToTelemetry/Private/CSVToTelemetry.cpp
+++ b/Engine/Source/Programs/CSVTools/CSVToTelemetry/Private/CSVToTelemetry.cpp
@@ -355,6 +355,9 @@ INT32_MAIN_INT32_ARGC_TCHAR_ARGV()
 	GEngineLoop.PreInit(ArgC, ArgV);
 
 	int32 Result = 1;
+	const float Version = 1.0;
+
+	UE_LOG(LogCSVToTelemetry, Display, TEXT("CSVToTelemetry v%.1f"), Version);
 
 	FString FilePath;
 
