PerfReportTool 4.263.0: Timing Region summary: Add arrows to indicate when region extents are outside of the track bounds
#Synced-CL 46049276

diff --git a/Engine/Source/Programs/CSVTools/PerfReportTool/PerfReportTool.cs b/Engine/Source/Programs/CSVTools/PerfReportTool/PerfReportTool.cs
--- a/Engine/Source/Programs/CSVTools/PerfReportTool/PerfReportTool.cs
+++ b/Engine/Source/Programs/CSVTools/PerfReportTool/PerfReportTool.cs
@@ -23,7 +23,7 @@ namespace PerfReportTool
     class Version
     {
 		// Format: Major.Minor.Bugfix
-        private static string VersionString = "4.262.2";
+        private static string VersionString = "4.263.0";
 
         public static string Get() { return VersionString; }
     };
diff --git a/Engine/Source/Programs/CSVTools/PerfReportTool/Summaries/TimingRegionSummary.cs b/Engine/Source/Programs/CSVTools/PerfReportTool/Summaries/TimingRegionSummary.cs
--- a/Engine/Source/Programs/CSVTools/PerfReportTool/Summaries/TimingRegionSummary.cs
+++ b/Engine/Source/Programs/CSVTools/PerfReportTool/Summaries/TimingRegionSummary.cs
@@ -313,8 +313,9 @@ namespace PerfSummaries
 				filterIndex = filterInfo.index;
 				trackWidth = inTrackWidth;
 				trackHeight = inTrackHeight;
+				innerTrackWidth = trackWidth - 25;
 				double trackDuration = (endTime - startTime);
-				regionDurationThreshold = (inFilterInfo.trackDisplayThreshold / trackWidth ) * trackDuration;
+				regionDurationThreshold = (inFilterInfo.trackDisplayThreshold / innerTrackWidth ) * trackDuration;
 			}
 			public bool TryAddEvent(RegionEvent beginEvent)
 			{
@@ -340,15 +341,17 @@ namespace PerfSummaries
 				// Strip the hidePrefix if necessary
 				string backgroundStyle = "fill:" + backgroundColor;
 				string textStyle = "font-family:verdana;font-size:12px";
+				string arrowTextStyle = "font-family:verdana;font-size:18px";
 				double trackDuration = endTime - startTime;
 				double lastEventX=xOffset;
+				RegionEvent lastEvent = null;
 				foreach (RegionEvent beginEvent in events)
 				{
 					RegionEvent endEvent = beginEvent.otherEvent;
 					double duration = endEvent.timestamp - beginEvent.timestamp;
 					double t = beginEvent.timestamp - startTime;
-					double x = Math.Round(xOffset + trackWidth * t / trackDuration,0);
-					double width = Math.Round(trackWidth * duration / trackDuration,0);
+					double x = Math.Round(xOffset + innerTrackWidth * t / trackDuration,0);
+					double width = Math.Round(innerTrackWidth * duration / trackDuration,0);
 					if (width < 3)
 					{
 						continue;
@@ -384,9 +387,24 @@ namespace PerfSummaries
 						htmlSection.WriteLine($"<text x='{x + 3}' y='{yOffset + 15}' style='{textStyle}' pointer-events='none'>{titleStr + regionText}</text>");
 					}
 					lastEventX = x + width;
+					lastEvent = endEvent;
 				}
 				// Redraw the background on the RHS to deal with any overflow text
-				htmlSection.WriteLine($"<rect x='{(lastEventX+1)}' y='{yOffset}' width='{(trackWidth-(lastEventX-xOffset)-1)}' height='{trackHeight}' style='{backgroundStyle}'/>");
+				htmlSection.WriteLine($"<rect x='{(lastEventX+1)}' y='{yOffset}' width='{(innerTrackWidth-(lastEventX-xOffset)-1)}' height='{trackHeight}' style='{backgroundStyle}'/>");
+				htmlSection.WriteLine($"<rect x='{xOffset + innerTrackWidth}' y='{yOffset}' width='{(trackWidth - innerTrackWidth)}' height='{trackHeight}' fill='#ffffff'/>");
+
+				// Draw left and right arrows to indicate regions which exceed the track
+				if (lastEvent != null && lastEvent.rangeType == ERangeType.After)
+				{
+					htmlSection.WriteLine($"<text x='{xOffset + innerTrackWidth + 5}' y='{yOffset + 15}' style='{arrowTextStyle}' pointer-events='none'>&#8680;</text>");
+				}
+				RegionEvent firstEvent = events.Count == 0 ? null : events[0];
+				if (firstEvent != null && firstEvent.rangeType == ERangeType.Before && firstEvent.timestamp > 0)
+				{
+					htmlSection.WriteLine($"<text x='{xOffset - 20}' y='{yOffset + 15}' style='{arrowTextStyle}' pointer-events='none'>&#8678;</text>");
+				}
+
+
 			}
 
 			public int CompareTo(RegionEventTrack other)
@@ -407,6 +425,7 @@ namespace PerfSummaries
 			public double startTime;
 			public double endTime;
 			public double trackWidth;
+			public double innerTrackWidth;
 			public double trackHeight;
 			public double regionDurationThreshold;
 			public RegionFilterInfo filterInfo;
@@ -443,7 +462,7 @@ namespace PerfSummaries
 			{
 				htmlSection = new HtmlSection(title, bStartCollapsed);
 
-				double imageWidth = 1400;
+				double imageWidth = 1425;
 				double headerWidth = 150;
 				double trackWidth = imageWidth - headerWidth;
 
@@ -488,6 +507,7 @@ namespace PerfSummaries
 
 				double yOffset = 0;
 				int trackIndex = 0;
+				double innerTrackWidth = imageWidth - 25;
 				foreach (RegionEventTrack track in tracks)
 				{
 					if (prevCategory != track.category)
@@ -495,7 +515,7 @@ namespace PerfSummaries
 						if (prevCategory != null)
 						{
 							double lineY = yOffset + 1;
-							htmlSection.WriteLine($"<line x1='{0}' y1='{lineY}' x2='{imageWidth}' y2='{lineY}' style='{lineStyle}' stroke-dasharray='5,5'/>");
+							htmlSection.WriteLine($"<line x1='{0}' y1='{lineY}' x2='{innerTrackWidth}' y2='{lineY}' style='{lineStyle}' stroke-dasharray='5,5'/>");
 							yOffset += 2;
 						}
 						prevCategory = track.category;
@@ -505,7 +525,9 @@ namespace PerfSummaries
 					yOffset += 20;
 					trackIndex++;
 				}
+				// Draw vertical dashed lines showing the edges of the tracks
 				htmlSection.WriteLine($"<line x1='{headerWidth-1}' y1='{0}' x2='{headerWidth-1}' y2='{imageHeight}' style='{lineStyle}' stroke-dasharray='5,5'/>");
+				htmlSection.WriteLine($"<line x1='{headerWidth+innerTrackWidth-headerWidth}' y1='{0}' x2='{headerWidth+innerTrackWidth-headerWidth}' y2='{imageHeight}' style='{lineStyle}' stroke-dasharray='5,5'/>");
 
 				htmlSection.WriteLine("</svg>");
 			}
