[PerformanceReportTool][Analytics] Adding telemetry emission capabilities for SummaryTableRow data.

#Synced-CL 49331470

diff --git a/Engine/Source/Programs/CSVTools/PerfReportTool/PerfReportTool.cs b/Engine/Source/Programs/CSVTools/PerfReportTool/PerfReportTool.cs
--- a/Engine/Source/Programs/CSVTools/PerfReportTool/PerfReportTool.cs
+++ b/Engine/Source/Programs/CSVTools/PerfReportTool/PerfReportTool.cs
@@ -23,7 +23,7 @@ namespace PerfReportTool
     class Version
     {
 		// Format: Major.Minor.Bugfix
-        private static string VersionString = "4.263.0";
+        private static string VersionString = "4.264.0";
 
         public static string Get() { return VersionString; }
     };
@@ -211,6 +211,7 @@ namespace PerfReportTool
 			"  -summaryTableToJsonFileStream : use a file stream to write Json. Experimental but can avoid OOMs\n" +
 			"  -summaryTableToJsonNoIndent : don't indent json output files\n" +
 			"  -jsonToPrcs <json filename> : write PRCs. PRC files will be written to -summaryTableCache folder\n" +
+			"  -summaryTableToJsonTelemetry : publish the json summary as a telemetry event. Requries -summaryTableToJson.\n" +
 			"\n" +
 			"Performance args for bulk mode:\n" +
 			"  -precacheCount <n> : number of CSV files to precache in the lookahead cache (0 for no precache)\n" +
@@ -447,6 +448,14 @@ namespace PerfReportTool
 			// If we're outputting row data to json, create the dict
 			SummaryTableDataJsonWriteHelper summaryTableJsonHelper = null;
 			string summaryJsonOutPath = GetArg("summaryTableToJson", null);
+			bool summaryTableJsonTelemetry = GetBoolArg("summaryTableToJsonTelemetry");
+
+			if(summaryTableJsonTelemetry && summaryJsonOutPath == null)
+			{
+				Console.WriteLine("-summaryTableToJsonTelemetry was specified, but not -summaryTableToJson. No telemetry will be emitted as json generation is a requirement.");
+				summaryTableJsonTelemetry = false;
+			}
+
 			if (summaryJsonOutPath != null)
 			{
 				summaryTableJsonHelper = new SummaryTableDataJsonWriteHelper(summaryJsonOutPath, 
@@ -719,7 +728,14 @@ namespace PerfReportTool
 								summaryTable.AddRowData(rowData, bReadAllStats, bShowHiddenStats);
 								if (summaryTableJsonHelper != null)
 								{
+
 									summaryTableJsonHelper.AddRowData(rowData);
+
+									// Emit telemetry if enabled
+									if(summaryTableJsonTelemetry)
+									{
+										Telemetry.EnqueuePostSummaryTableMetricEvent(rowData);
+									}
 								}
 							}
 							perfLog.LogTiming("  AddRowData");
@@ -749,6 +765,13 @@ namespace PerfReportTool
 			{
 				summaryTableJsonHelper.WriteToJson(GetBoolArg("summaryTableToJsonFileStream"));
 				perfLog.LogTiming("WriteSummaryDataJson");
+
+				if (summaryTableJsonTelemetry)
+				{
+					Telemetry.Flush();
+					Telemetry.Dispose();
+				}
+
 				if (GetBoolArg("summaryTableToJsonFastMode"))
 				{
 					perfLog.LogTotalTiming();
@@ -2401,4 +2424,3 @@ namespace PerfReportTool
 		}
 	}
 }
-
diff --git a/Engine/Source/Programs/CSVTools/PerfReportTool/PerfreportTool.csproj b/Engine/Source/Programs/CSVTools/PerfReportTool/PerfreportTool.csproj
--- a/Engine/Source/Programs/CSVTools/PerfReportTool/PerfreportTool.csproj
+++ b/Engine/Source/Programs/CSVTools/PerfReportTool/PerfreportTool.csproj
@@ -1,4 +1,4 @@
-<Project Sdk="Microsoft.NET.Sdk">
+﻿<Project Sdk="Microsoft.NET.Sdk">
   <PropertyGroup>
     <TargetFramework>net8.0</TargetFramework>
     <OutputType>Exe</OutputType>
@@ -17,6 +17,7 @@
     <None Remove="Properties\AssemblyInfo.cs~RF273305d.TMP" />
   </ItemGroup>
   <ItemGroup>
+    <ProjectReference Include="..\..\Shared\EpicGames.Horde\EpicGames.Horde.csproj" />
     <ProjectReference Include="..\CsvStats\CsvStats.csproj" />
     <ProjectReference Include="..\CsvToSvgLib\CsvToSvgLib.csproj" />
   </ItemGroup>
diff --git a/Engine/Source/Programs/CSVTools/PerfReportTool/Telemetry.cs b/Engine/Source/Programs/CSVTools/PerfReportTool/Telemetry.cs
new file mode 100644
--- /dev/null
+++ b/Engine/Source/Programs/CSVTools/PerfReportTool/Telemetry.cs
@@ -0,0 +1,271 @@
+// Copyright Epic Games, Inc. All Rights Reserved.
+
+using EpicGames.Core;
+using EpicGames.Horde;
+using Microsoft.Extensions.DependencyInjection;
+using PerfSummaries;
+using System;
+using System.Collections;
+using System.Collections.Generic;
+using System.Net.Http;
+using System.Threading.Tasks;
+
+namespace PerfReportTool
+{
+	/// <summary>
+	/// Telemetry class used to emit telemetry events for PerfReportTool.
+	/// </summary>
+	internal static class Telemetry
+	{
+		private static object _lock = new object();
+		private const string DefaultAppId = "Horde";
+		private const int DefaultAppVersion = 1;
+		private const string DefaultAppEnvironment = "Dev";
+		private const string EventName = "State.Performance.Debug";
+		private static ServiceProvider _serviceProvider = null;
+		private static IHordeClient _hordeClient = null;
+		private static HordeHttpClient _hordeHttpClient = null;
+		private readonly static List<Task<bool>> _queuedRequests = new List<Task<bool>>(16);
+
+		/// <summary>
+		/// Gets the HordeClient used to create Horde Http Clients.
+		/// </summary>
+		internal static IHordeClient HordeClient
+		{
+			get
+			{
+				if (_hordeClient == null)
+				{
+					if (_serviceProvider == null)
+					{
+						Uri defaultServer = HordeOptions.GetDefaultServerUrl() ?? HordeOptions.GetServerUrlFromEnvironment();
+
+						if (defaultServer == null)
+						{
+							return null;
+						}
+
+						ServiceCollection services = new();
+						services.AddHorde(o => o.ServerUrl = defaultServer);
+						_serviceProvider = services.BuildServiceProvider();
+					}
+
+					_hordeClient = _serviceProvider.GetRequiredService<IHordeClient>();
+				}
+
+				return _hordeClient;
+			}
+		}
+
+		/// <summary>
+		/// Gets the Horde Http Client used to emit events.
+		/// </summary>
+		internal static HordeHttpClient HordeHttpClient
+		{
+			get
+			{
+				if (_hordeHttpClient == null)
+				{
+					IHordeClient client = HordeClient;
+					if (client != null)
+					{
+						_hordeHttpClient = HordeClient.CreateHttpClient();
+					}
+				}
+
+				return _hordeHttpClient;
+			}
+		}
+
+		/// <summary>
+		/// Flushes all telemetry events that have been queued.
+		/// </summary>
+		internal static void Flush()
+		{
+			Console.WriteLine("Flushing {0} telemetry events.", _queuedRequests.Count);
+
+			if (_queuedRequests.Count == 0)
+			{
+				return;
+			}
+
+			Task[] tasks;
+			lock (_lock)
+			{
+				tasks = _queuedRequests.ToArray();
+				_queuedRequests.Clear();
+			}
+
+			Task.WhenAll(tasks).GetAwaiter().GetResult();
+		}
+
+		/// <summary>
+		/// Releases underyling singletons for the Telemetry helper.
+		/// </summary>
+		internal static void Dispose()
+		{
+			if (_serviceProvider != null)
+			{
+				_serviceProvider.DisposeAsync().AsTask().GetAwaiter().GetResult();
+				_serviceProvider = null;
+
+				// This is disposed as a part of the service provider.
+				_hordeClient = null;
+			}
+
+			if (_hordeHttpClient != null)
+			{
+				_hordeHttpClient.Dispose();
+				_hordeHttpClient = null;
+			}
+		}
+
+		/// <summary>
+		/// Posts a row data's SummaryTableMetric to the Horde event stream, asynchronously.
+		/// </summary>
+		/// <param name="rowData">The row data to post.</param>
+		/// <returns>True if telemetry was emitted successfully, false otherwise.</returns>
+		internal async static Task<bool> PostSummaryTableMetricEventAsync(SummaryTableRowData rowData)
+		{
+			if (HordeClient == null)
+			{
+				return false;
+			}
+
+			Hashtable telemetryPayload = null;
+
+			try
+			{
+				telemetryPayload = CreatePayload(rowData);
+			}
+			catch (InvalidOperationException ex)
+			{
+				Console.WriteLine("Skipping emission of rowData due to invalid payload construction. Exception: {0}", ex.Message);
+				return false;
+			}
+
+			if (telemetryPayload == null)
+			{
+				return false;
+			}
+
+			try
+			{
+				using HttpResponseMessage response = await HordeHttpClient.PostTelemetryAsync(DefaultAppId, DefaultAppVersion, DefaultAppEnvironment, [telemetryPayload]);
+
+				if (response == null)
+				{
+					return false;
+				}
+
+				bool isSuccessStatusCode = response.IsSuccessStatusCode;
+				return isSuccessStatusCode;
+			}
+			catch (Exception ex)
+			{
+				Console.WriteLine("Encountered an exception while posting event. Exception: {0}", ex.Message);
+				return false;
+			}
+		}
+
+		/// <summary>
+		/// Enqueues a task to posts a row data's SummaryTableMetric to the Horde event stream.
+		/// </summary>
+		/// <param name="rowData">The row data to post.</param>
+		/// <returns>A readonly list containing the queued tasks.</returns>
+		internal static IReadOnlyList<Task<bool>> EnqueuePostSummaryTableMetricEvent(SummaryTableRowData rowData)
+		{
+			lock (_lock)
+			{
+				_queuedRequests.Add(PostSummaryTableMetricEventAsync(rowData));
+			}
+
+			return _queuedRequests;
+		}
+
+		/// <summary>
+		/// Posts a row data's SummaryTableMetric to the Horde event stream.
+		/// </summary>
+		/// <param name="rowData">The row data to post.</param>
+		/// <returns>True if telemetry was emitted successfully, false otherwise.</returns>
+		internal static bool PostSummaryTableMetricEvent(SummaryTableRowData rowData)
+		{
+			if (HordeClient == null)
+			{
+				return false;
+			}
+
+			Hashtable telemetryPayload = null;
+
+			try
+			{
+				telemetryPayload = CreatePayload(rowData);
+			}
+			catch (InvalidOperationException ex)
+			{
+				Console.WriteLine("Skipping emission of rowData due to invalid payload construction. Exception: {0}", ex.Message);
+				return false;
+			}
+
+			if (telemetryPayload == null)
+			{
+				return false;
+			}
+
+			try
+			{
+				using HttpResponseMessage response = HordeHttpClient.PostTelemetryAsync(DefaultAppId, DefaultAppVersion, DefaultAppEnvironment, [telemetryPayload]).GetAwaiter().GetResult();
+				bool isSuccessStatusCode = response.IsSuccessStatusCode;
+				return isSuccessStatusCode;
+			}
+			catch (Exception ex)
+			{
+				Console.WriteLine("Encountered an exception while posting event. Exception: {0}", ex.Message);
+				return false;
+			}
+		}
+
+		private static Hashtable CreatePayload(SummaryTableRowData rowData)
+		{
+			Dictionary<string, dynamic> jsonDict = rowData.ToJsonDict(false, false);
+			string summaryTableMetricKey = SummaryTableElement.Type.SummaryTableMetric.ToString();
+			dynamic summaryTableMetric;
+
+			if (!jsonDict.TryGetValue(summaryTableMetricKey, out summaryTableMetric))
+			{
+				return null;
+			}
+
+			if (summaryTableMetric == null)
+			{
+				return null;
+			}
+
+			Hashtable telemetryPayload = [];
+			telemetryPayload["EventName"] = EventName;
+
+			// Coalesce summary table metric values into telemetry payload via IDictionary || IDictionary<string,TValue> interfaces.
+			if (summaryTableMetric is IDictionary)
+			{
+				foreach (dynamic kvp in summaryTableMetric)
+				{
+					telemetryPayload[kvp.Key] = kvp.Value;
+				}
+			}
+
+			else if (summaryTableMetric is IDictionary<string, object> stringKeyGenericDict)
+			{
+				foreach (KeyValuePair<string, object> kvp in stringKeyGenericDict)
+				{
+					telemetryPayload[kvp.Key] = kvp.Value;
+				}
+			}
+			else
+			{
+				throw new InvalidOperationException("SummaryTableMetric must be a dictionary or Dictionary<string, object>.");
+			}
+
+			return telemetryPayload;
+		}
+	}
+}
\ No newline at end of file
