[PerformanceReportTool][Analytics] Adding metadata & specific summary table emissions

#Synced-CL 49562268

diff --git a/Engine/Plugins/Performance/AutomatedPerfTesting/Build/Scripts/PerfReport/ReportTypes.xml b/Engine/Plugins/Performance/AutomatedPerfTesting/Build/Scripts/PerfReport/ReportTypes.xml
--- a/Engine/Plugins/Performance/AutomatedPerfTesting/Build/Scripts/PerfReport/ReportTypes.xml
+++ b/Engine/Plugins/Performance/AutomatedPerfTesting/Build/Scripts/PerfReport/ReportTypes.xml
@@ -125,6 +125,10 @@
 			<filter>Platform,DeviceProfile,scalabilitylevel,Report,mapsequencecomboname,cameracut,Horde Job,TargetFPS,MVP*,Hitches/Min,HitchTimePercent,FrameTime*,GamethreadTime*,Renderthreadtime*,RHIThreadTime*,GPUTime*,MemoryFreeMB Min,PhysicalUsedMB Max,DynamicResolution*,RHI/Drawcalls*,Basic/TicksQueued*,Total Time*,Csv File,buildversion,config,playlistid,framecount*,syncloads</filter>
 		</summaryTable>
 
+		<summaryTable name="keystats" rowSort="buildversion,Platform">
+			<filter>Platform,MVP*,Hitches/Min,HitchTimePercent,GamethreadTime*,DynamicResolution*,Renderthreadtime*,GPUTime*,PhysicalUsedMB Max</filter>
+		</summaryTable>
+
 		<summaryTable name="staticcamera" rowSort="buildversion,Platform,mapname,cameraname,deviceprofile,scalabilitylevel,TargetFPS">
 			<filter>Platform,DeviceProfile,mapname,cameraname,scalabilitylevel,Report,Horde Job,TargetFPS,MVP*,Hitches/Min,HitchTimePercent,FrameTime*,GamethreadTime*,Renderthreadtime*,RHIThreadTime*,GPUTime*,MemoryFreeMB Min,PhysicalUsedMB Max,DynamicResolution*,RHI/Drawcalls*,Basic/TicksQueued*,Total Time*,Csv File,buildversion,config,playlistid,framecount*,syncloads</filter>
 		</summaryTable>
diff --git a/Engine/Source/Programs/CSVTools/PerfReportTool/PerfReportTool.cs b/Engine/Source/Programs/CSVTools/PerfReportTool/PerfReportTool.cs
--- a/Engine/Source/Programs/CSVTools/PerfReportTool/PerfReportTool.cs
+++ b/Engine/Source/Programs/CSVTools/PerfReportTool/PerfReportTool.cs
@@ -17,20 +17,21 @@ using PerfSummaries;
 using CSVTools;
 
 using System.Text.Json;
+using EpicGames.Analytics.PerformanceTrends;
 
 namespace PerfReportTool
 {
-    class Version
-    {
+	class Version
+	{
 		// Format: Major.Minor.Bugfix
-        private static string VersionString = "4.264.0";
+		private static string VersionString = "4.265.0";
 
-        public static string Get() { return VersionString; }
-    };
+		public static string Get() { return VersionString; }
+	};
 
 	class HashHelper
 	{
-		public static string StringToHashStr(string strIn, int maxCharsOut=-1)
+		public static string StringToHashStr(string strIn, int maxCharsOut = -1)
 		{
 			HashAlgorithm algorithm = SHA256.Create();
 			StringBuilder sb = new StringBuilder();
@@ -69,7 +70,7 @@ namespace PerfReportTool
 			}
 			if (HitCount > 0 || MissCount > 0)
 			{
-				Console.WriteLine("  Hit percentage  : " + ((float)HitCount * 100.0f / ((float)MissCount+(float)HitCount)).ToString("0.0") + "%");
+				Console.WriteLine("  Hit percentage  : " + ((float)HitCount * 100.0f / ((float)MissCount + (float)HitCount)).ToString("0.0") + "%");
 			}
 		}
 	};
@@ -114,7 +115,7 @@ namespace PerfReportTool
 			"  -embeddedGraphUrlRoot : The url to fetch the graph from if -useEmbeddedGraphUrl is specified. CsvToSvg graph args are provided as get params\n" +
 			"Deprecated performance args:\n" +
 			"  -csvToSvgProcesses : Use separate processes for csvToSVG instead of threads (slower)\n" +
-            "  -embedGraphCommandline : if -csvToSvgProcesses is specified, embeds the commandline for debugging purposes\n" +
+			"  -embedGraphCommandline : if -csvToSvgProcesses is specified, embeds the commandline for debugging purposes\n" +
 			"  -noBatchedGraphs : disable batched/multithreaded graph generation (use with -csvToSvgProcesses. Default is enabled)\n" +
 			"\n" +
 			"Options to truncate or filter source data:\n" +
@@ -151,7 +152,7 @@ namespace PerfReportTool
 			"  -spreadsheetfriendly: outputs a single quote before non-numeric entries in summary tables\n" +
 			"  -noSummaryMinMax: don't make min/max columns for each stat in a condensed summary\n" +
 			"  -reverseTable [0|1]: Reverses the order of summary tables (set 0 to force off)\n" +
-			"  -sortTrailingDigitsAsNumeric : detects trailing digits and pads them when sorting summary table\n"+
+			"  -sortTrailingDigitsAsNumeric : detects trailing digits and pads them when sorting summary table\n" +
 			"  -scrollableTable [0|1]: makes the summary table scrollable, with frozen first rows and columns (set 0 to force off)\n" +
 			"  -colorizeTable <off|budget|auto>: selects the table colorization mode. If omitted, uses the default in the summary\n" +
 			"     xml table if set.\n" +
@@ -181,14 +182,14 @@ namespace PerfReportTool
 			"  -collatedDateVisibility <hide|newest|oldest> : sets visibility of dates in collated summary tables (default=newest)\n" +
 			"  -columnSortMode <default|sortByMax|sortByAvg|wildCardSortByMax|wildCardSortByAvg> : sorting mode for columns\n" +
 			"  -summaryTableJsData : embeds javascript row and column data in the <script> section to help with interop\n" +
-            "  -derivedStatVerboseErrors : output verbose errors for derived stat expressions\n" +
+			"  -derivedStatVerboseErrors : output verbose errors for derived stat expressions\n" +
 
 			"Diff rows\n" +
 			"  -addDiffRows : adds diff rows after the first two rows\n" +
 			"  -sortColumnsByDiff : sorts columns by the max of its diff scores (use with -addDiffRows).\n" +
 			"      Notes: Diff score corresponds to the value of a column's diff row; the sign is reversed if LowIsBad\n" +
 			"             Stats are sorted by prefixes first, e.g GPU/\n" +
-			"  -columnDiffDisplayThreshold <value> : if specified, hides columns with max diff value below this threshold\n"+
+			"  -columnDiffDisplayThreshold <value> : if specified, hides columns with max diff value below this threshold\n" +
 			"  -diffRowsAlternating 0|1. If enabled, every other column has a diff. Otherwise diffs are only generated for pairs\n" +
 			"  -showOnlyDiffRows: If enabled, non-diff rows are hidden. Common text fields are combined.\n" +
 			"\n" +
@@ -198,9 +199,9 @@ namespace PerfReportTool
 			"Regression Column Filtering\n" +
 			"  -onlyShowRegressedColumns : enables regression filtering. Only shows columns where the most recent row group\n" +
 			"    (see -regressionJoinRowsByName) is outside the given stddiv threshold from the mean of the previous rows.\n" +
-			"  -regressionJoinRowsByName <statName> : a stat name to join rows by for aggregation. (default: no aggregation)\n"+
-			"  -regressionStdDevThreshold <n> (default = 2) : the stddiv threshold for filtering \n"+
-			"  -regressionOutlierStdDevThreshold <n> (default = 4) : stddiv threshold for outliers (these are ignored)\n"+
+			"  -regressionJoinRowsByName <statName> : a stat name to join rows by for aggregation. (default: no aggregation)\n" +
+			"  -regressionStdDevThreshold <n> (default = 2) : the stddiv threshold for filtering \n" +
+			"  -regressionOutlierStdDevThreshold <n> (default = 4) : stddiv threshold for outliers (these are ignored)\n" +
 			"\n" +
 			"Json serialization:\n" +
 			"  -summaryTableToJson <path> : path (usually a json filename) to write summary table row data to\n" +
@@ -211,7 +212,10 @@ namespace PerfReportTool
 			"  -summaryTableToJsonFileStream : use a file stream to write Json. Experimental but can avoid OOMs\n" +
 			"  -summaryTableToJsonNoIndent : don't indent json output files\n" +
 			"  -jsonToPrcs <json filename> : write PRCs. PRC files will be written to -summaryTableCache folder\n" +
-			"  -summaryTableToJsonTelemetry : publish the json summary as a telemetry event. Requries -summaryTableToJson.\n" +
+			"Telemetry Emission:\n" +
+			"  -summaryTableDataToTelemetry : publish the summary table as a telemetry event.\n" +
+			"  -summaryTableDataToTelemetryVerbose : publish the verbose form of every row data (without summary table filter).\n" +
+			"     Will also emit summary table data.\n" +
 			"\n" +
 			"Performance args for bulk mode:\n" +
 			"  -precacheCount <n> : number of CSV files to precache in the lookahead cache (0 for no precache)\n" +
@@ -352,7 +356,7 @@ namespace PerfReportTool
 					bBulkMode = true;
 					bSummaryTableCacheOnlyMode = true;
 				}
-				else if (metadataProxyName.Length > 0 || explicitMetadata.Length> 0)
+				else if (metadataProxyName.Length > 0 || explicitMetadata.Length > 0)
 				{
 					string proxyMetadata = "";
 
@@ -374,11 +378,11 @@ namespace PerfReportTool
 					proxyMetadata += explicitMetadata;
 
 					// Process each key=value pair, later entries will override previous ones
-					Dictionary<string,string> metaDict = new Dictionary<string,string>();
+					Dictionary<string, string> metaDict = new Dictionary<string, string>();
 					if (proxyMetadata.Length > 0)
 					{
 						// This allows flexibility on the commandline, can use commas or semicolons (and spaces since they are mapped to semicolons in GetArg)
-						string[] entries = proxyMetadata.Split(',',';');
+						string[] entries = proxyMetadata.Split(',', ';');
 						foreach (string entry in entries)
 						{
 							string[] data = entry.Split("=");
@@ -448,24 +452,19 @@ namespace PerfReportTool
 			// If we're outputting row data to json, create the dict
 			SummaryTableDataJsonWriteHelper summaryTableJsonHelper = null;
 			string summaryJsonOutPath = GetArg("summaryTableToJson", null);
-			bool summaryTableJsonTelemetry = GetBoolArg("summaryTableToJsonTelemetry");
 
-			if(summaryTableJsonTelemetry && summaryJsonOutPath == null)
-			{
-				Console.WriteLine("-summaryTableToJsonTelemetry was specified, but not -summaryTableToJson. No telemetry will be emitted as json generation is a requirement.");
-				summaryTableJsonTelemetry = false;
-			}
+			bool summaryTableDataToTelemetryVerbose = GetBoolArg("summaryTableDataToTelemetryVerbose");
+			bool summaryTableDataToTelemetry = GetBoolArg("summaryTableDataToTelemetry") || summaryTableDataToTelemetryVerbose;
 
 			if (summaryJsonOutPath != null)
 			{
-				summaryTableJsonHelper = new SummaryTableDataJsonWriteHelper(summaryJsonOutPath, 
-					GetBoolArg("summaryTableToJsonSeparateFiles"), 
-					GetBoolArg("summaryTableToJsonMetadataOnly"), 
+				summaryTableJsonHelper = new SummaryTableDataJsonWriteHelper(summaryJsonOutPath,
+					GetBoolArg("summaryTableToJsonSeparateFiles"),
+					GetBoolArg("summaryTableToJsonMetadataOnly"),
 					GetBoolArg("summaryTableToJsonWriteAllElementData"),
 					!GetBoolArg("summaryTableToJsonNoIndent"));
 			}
 
-
 			SummaryTableCacheStats summaryTableCacheStats = new SummaryTableCacheStats();
 
 			perfLog.LogTiming("Initialization");
@@ -730,14 +729,15 @@ namespace PerfReportTool
 								{
 
 									summaryTableJsonHelper.AddRowData(rowData);
-
-									// Emit telemetry if enabled
-									if(summaryTableJsonTelemetry)
-									{
-										Telemetry.EnqueuePostSummaryTableMetricEvent(rowData);
-									}
 								}
 							}
+
+							// Emit telemetry if enabled
+							if (summaryTableDataToTelemetryVerbose)
+							{
+								Telemetry.EnqueuePostSummaryTableMetricEvent(rowData);
+							}
+
 							perfLog.LogTiming("  AddRowData");
 						}
 					}
@@ -766,14 +766,14 @@ namespace PerfReportTool
 				summaryTableJsonHelper.WriteToJson(GetBoolArg("summaryTableToJsonFileStream"));
 				perfLog.LogTiming("WriteSummaryDataJson");
 
-				if (summaryTableJsonTelemetry)
-				{
-					Telemetry.Flush();
-					Telemetry.Dispose();
-				}
-
 				if (GetBoolArg("summaryTableToJsonFastMode"))
 				{
+					if (summaryTableDataToTelemetry)
+					{
+						Telemetry.Flush();
+						Telemetry.Dispose();
+					}
+
 					perfLog.LogTotalTiming();
 					return;
 				}
@@ -798,13 +798,13 @@ namespace PerfReportTool
 				bool bCollateTableOnly = GetBoolArg("collateTableOnly");
 
 				bCollateTable |= bCollateTableOnly;
-				string collatedTableFilename = summaryTableFilename + ( bCollateTableOnly ? "" : "_Collated" );
+				string collatedTableFilename = summaryTableFilename + (bCollateTableOnly ? "" : "_Collated");
 
 				bool bSpreadsheetFriendlyStrings = GetBoolArg("spreadsheetFriendly");
 				string weightByColumnName = GetArg("weightByColumn", null);
 
 				string customSummaryTableRowSort = GetArg("customTableSort", null);
-				string summaryTableName = GetArg("summaryTable",null);
+				string summaryTableName = GetArg("summaryTable", null);
 				SummaryTableInfo tableInfo = GetTableInfo(summaryTableName, customSummaryTableFilter, customSummaryTableRowSort);
 				if (!bCollateTableOnly)
 				{
@@ -819,7 +819,7 @@ namespace PerfReportTool
 				string condensedSummaryTableName = GetArg("condensedSummaryTable", null);
 				if (GetBoolArg("emailSummary") || GetBoolArg("emailTable") || condensedSummaryTableName != null)
 				{
-					if (condensedSummaryTableName == null )
+					if (condensedSummaryTableName == null)
 					{
 						condensedSummaryTableName = "condensed";
 					}
@@ -829,6 +829,12 @@ namespace PerfReportTool
 				perfLog.LogTiming("WriteSummaryTable");
 			}
 
+			if (summaryTableDataToTelemetry)
+			{
+				Telemetry.Flush();
+				Telemetry.Dispose();
+			}
+
 			if (summaryTableCacheDir != null)
 			{
 				summaryTableCacheStats.LogStats();
@@ -842,7 +848,7 @@ namespace PerfReportTool
 			SummaryTableInfo tableInfo;
 
 			// Use the default summary table name if one isn't provided, unless we're making a custom table			
-			if (summaryTableName == null && customSummaryTableFilterString == null )
+			if (summaryTableName == null && customSummaryTableFilterString == null)
 			{
 				summaryTableName = "default";
 			}
@@ -858,7 +864,7 @@ namespace PerfReportTool
 			if (customSummaryTableFilterString != null)
 			{
 				tableInfo.columnFilterList = customSummaryTableFilterString.Split(',').ToList();
-				
+
 				// If we're doing a custom sort but don't have a provided row sort (or preset summary table), set a default
 				if (summaryTableName == null && customSummaryTableRowSortString == null)
 				{
@@ -911,7 +917,7 @@ namespace PerfReportTool
 			// If -colorizeTable isn't specified then we use the default from the xml.
 			// If -colorizeTable isn't specified and it's not set in the summary table xml it uses the default set in the class initializer.
 			TableColorizeMode colorizeMode = tableInfo.tableColorizeMode;
-			if ( GetBoolArg("autoColorizeTable")) // Legacy support for the -autoColorizeTable arg
+			if (GetBoolArg("autoColorizeTable")) // Legacy support for the -autoColorizeTable arg
 			{
 				colorizeMode = TableColorizeMode.Auto;
 			}
@@ -958,22 +964,41 @@ namespace PerfReportTool
 				filteredTable.SetColumnFormatInfo(reportXML.columnFormatInfoList);
 
 				filteredTable.WriteToHTML(
-					filenameWithoutExtension + ".html", 
-					VersionString, 
-					bSpreadsheetFriendlyStrings, 
-					tableInfo.sectionBoundaries, 
+					filenameWithoutExtension + ".html",
+					VersionString,
+					bSpreadsheetFriendlyStrings,
+					tableInfo.sectionBoundaries,
 					bScrollableTable,
-					colorizeMode, 
-					addMinMaxColumns, 
+					colorizeMode,
+					addMinMaxColumns,
 					tableInfo.hideStatPrefix,
-					GetIntArg("maxSummaryTableStringLength", Int32.MaxValue), 
-					weightByColumnName, 
+					GetIntArg("maxSummaryTableStringLength", Int32.MaxValue),
+					weightByColumnName,
 					summaryTitle,
 					bCollated ? bTransposeCollatedSummaryTable : bTransposeFullSummaryTable,
 					tableInfo.summaryTableStyle,
 					GetBoolArg("summaryTableJsData")
 				);
 			}
+
+			bool summaryTableDataToTelemetry = GetBoolArg("summaryTableDataToTelemetry") || GetBoolArg("summaryTableDataToTelemetryVerbose");
+
+			// Emit telemetry if enabled
+			if (summaryTableDataToTelemetry)
+			{
+				// Reconstitute the table to include our standard metadata columns for each bundle of data.
+				SummaryTable reconstitutedTable = table.SortAndFilter(tableInfo.columnFilterList.Union(StandardTelemetryMetadata.StandardMetadataKeys).ToList(), tableInfo.rowSortList, bReverseTable, weightByColumnName, additionalColumnFilters, bSortTrailingDigitsAsNumeric, tableColumnSortMode);
+
+				if (bCollated)
+				{
+					StringCollationVisibility stringCollationVisibility = GetEnumArg("collatedStringVisibility", tableInfo.stringCollationVisibility);
+					DateCollationVisibility dateCollationVisibility = GetEnumArg("collatedDateVisibility", tableInfo.dateCollationVisibility);
+					reconstitutedTable = reconstitutedTable.CollateSortedTable(tableInfo.rowSortList, addMinMaxColumns, stringCollationVisibility, dateCollationVisibility);
+				}
+
+				string summaryTableName = GetArg("summaryTable", null);
+				Telemetry.EnqueuePostSummaryTableMetricEvent(summaryTableName, reconstitutedTable, bCollated);
+			}
 		}
 
 		IEnumerable<ISummaryTableColumnFilter> MakeAdditionalColumnFilters(SummaryTableInfo tableInfo)
@@ -1699,12 +1724,12 @@ namespace PerfReportTool
 					{
 						htmlFile.WriteLine("<h4>" + currentCategory + " Graphs</h4>");
 					}
-					
+
 					foreach (CsvSvgInfo csvSvgInfo in csvSvgInfoList)
 					{
 						string svgTitle = csvSvgInfo.Graph.title;
 						// TODO: Check if this graph belongs in this section.
-						htmlFile.WriteLine("<li><a href='#" + StripSpaces(svgTitle) + "' onclick='setSectionExpandedById(\""+ StripSpaces(svgTitle) + "\", true)'>" + svgTitle + "</a></li>");
+						htmlFile.WriteLine("<li><a href='#" + StripSpaces(svgTitle) + "' onclick='setSectionExpandedById(\"" + StripSpaces(svgTitle) + "\", true)'>" + svgTitle + "</a></li>");
 					}
 
 					htmlFile.WriteLine("</ul>");
@@ -1776,7 +1801,7 @@ namespace PerfReportTool
 					{
 						throw new Exception("Unsupported graph output format.");
 					}
-					
+
 					htmlSection.WriteToFile(htmlFile);
 				}
 
@@ -1878,7 +1903,7 @@ namespace PerfReportTool
 		}
 
 		enum CsvToSvgArgFormat
-		{ 
+		{
 			CommandLine,
 			Url
 		}
@@ -2016,7 +2041,7 @@ namespace PerfReportTool
 			AddConditionalArg("endEvent", graphSettings.endEvent.isSet, graphSettings.endEvent.value);
 			AddConditionalArg("endEventOffset", graphSettings.endEventOffset.isSet, graphSettings.endEventOffset.value);
 
-			List<string> statList = CsvStats.SplitStringListWithBracketGroups(graphSettings.statString.value,',');
+			List<string> statList = CsvStats.SplitStringListWithBracketGroups(graphSettings.statString.value, ',');
 
 			string argString = string.Empty;
 			if (argFormat == CsvToSvgArgFormat.Url)
@@ -2032,7 +2057,7 @@ namespace PerfReportTool
 				// Exclude csvs from the args as that is a separate argument that the caller must setup.
 				// Any multi-value args need to be delimited with a semi-colon.
 				// Derived csv stat definitions can include equals, so encode it
-				args["stats"] = String.Join(";", statList).Replace("=", "%3D"); 
+				args["stats"] = String.Join(";", statList).Replace("=", "%3D");
 				AddConditionalArg("showevents", showEvents.Length > 0, FormatStringList(showEvents, " "));
 				AddConditionalArg("highlightEventRegions", highlightEventRegions.Length > 0, FormatStringList(highlightEventRegions, ","));
 
@@ -2070,7 +2095,7 @@ namespace PerfReportTool
 			graphParams.title = graph.title;
 
 			GraphSettings graphSettings = graph.settings;
-			graphParams.statNames.AddRange(CsvStats.SplitStringListWithBracketGroups(graphSettings.statString.value,','));
+			graphParams.statNames.AddRange(CsvStats.SplitStringListWithBracketGroups(graphSettings.statString.value, ','));
 			graphParams.lineThickness = (float)(graphSettings.thickness.value * thicknessMultiplier);
 			graphParams.smooth = graphSettings.smooth.value && !GetBoolArg("nosmooth");
 			if (graphParams.smooth)
@@ -2153,7 +2178,7 @@ namespace PerfReportTool
 				if (graphSettings.startEventOffset.isSet)
 				{
 					graphParams.startEventOffset = graphSettings.startEventOffset.value;
-				}					
+				}
 			}
 
 			if (graphSettings.endEvent.isSet)
@@ -2244,7 +2269,7 @@ namespace PerfReportTool
 					binary = binary.Replace("\\perfreporttool\\", "\\CsvToSVG\\", StringComparison.InvariantCultureIgnoreCase);
 				}
 			}
-			else 
+			else
 			{
 				throw new NotImplementedException("-CsvToSvgProcesses is deprecated and only supported on Windows");
 			}
@@ -2358,7 +2383,7 @@ namespace PerfReportTool
 			{
 				Directory.CreateDirectory(prcOutputDir);
 			}
-			Console.WriteLine("Reading "+jsonFilename);
+			Console.WriteLine("Reading " + jsonFilename);
 			string jsonText = File.ReadAllText(jsonFilename);
 
 			Console.WriteLine("Parsing json");
@@ -2406,7 +2431,7 @@ namespace PerfReportTool
 
 				case JsonValueKind.Array:
 					List<dynamic> ArrayValue = new List<dynamic>();
-					foreach( JsonElement element in jsonElement.EnumerateArray())
+					foreach (JsonElement element in jsonElement.EnumerateArray())
 					{
 						ArrayValue.Add(GetJsonValue(element));
 					}
@@ -2414,7 +2439,7 @@ namespace PerfReportTool
 
 				case JsonValueKind.Object:
 					Dictionary<string, dynamic> DictValue = new Dictionary<string, dynamic>();
-					foreach ( JsonProperty property in jsonElement.EnumerateObject() )
+					foreach (JsonProperty property in jsonElement.EnumerateObject())
 					{
 						DictValue[property.Name] = GetJsonValue(property.Value);
 					}
diff --git a/Engine/Source/Programs/CSVTools/PerfReportTool/PerfreportTool.csproj b/Engine/Source/Programs/CSVTools/PerfReportTool/PerfreportTool.csproj
--- a/Engine/Source/Programs/CSVTools/PerfReportTool/PerfreportTool.csproj
+++ b/Engine/Source/Programs/CSVTools/PerfReportTool/PerfreportTool.csproj
@@ -17,6 +17,7 @@
     <None Remove="Properties\AssemblyInfo.cs~RF273305d.TMP" />
   </ItemGroup>
   <ItemGroup>
+    <ProjectReference Include="..\..\Shared\EpicGames.Analytics.PerformanceTrends\EpicGames.Analytics.PerformanceTrends.csproj" />
     <ProjectReference Include="..\..\Shared\EpicGames.Horde\EpicGames.Horde.csproj" />
     <ProjectReference Include="..\CsvStats\CsvStats.csproj" />
     <ProjectReference Include="..\CsvToSvgLib\CsvToSvgLib.csproj" />
diff --git a/Engine/Source/Programs/CSVTools/PerfReportTool/SummaryTable.cs b/Engine/Source/Programs/CSVTools/PerfReportTool/SummaryTable.cs
--- a/Engine/Source/Programs/CSVTools/PerfReportTool/SummaryTable.cs
+++ b/Engine/Source/Programs/CSVTools/PerfReportTool/SummaryTable.cs
@@ -2011,6 +2011,29 @@ namespace PerfSummaries
 			csvFile.Close();
 		}
 
+		/// <summary>
+		/// An enumerable representation of the SummaryTable, by rows.
+		/// </summary>
+		/// <returns>The row of data.</returns>
+		public IEnumerable<IEnumerable<KeyValuePair<SummaryTableColumn, dynamic>>> EnumerateTableDataByRow()
+		{
+			for (int i = 0; i < rowCount; i++)
+			{
+				IList<KeyValuePair<SummaryTableColumn,dynamic>> row = new List<KeyValuePair<SummaryTableColumn, dynamic>>(columns.Count);
+
+				foreach (var col in columns)
+				{
+					dynamic value = col.isNumeric
+						? col.GetValue(i)
+						: col.GetStringValue(i);
+
+					row.Add(new KeyValuePair<SummaryTableColumn, dynamic>(col, value));
+				}
+
+				yield return row;
+			}
+		}
+
 		// Sorts columns by their diff score. 
 		// Requires max rows. Also works with collated tables with min/max
 		private void SortColumnsByDiffRows()
diff --git a/Engine/Source/Programs/CSVTools/PerfReportTool/Telemetry.cs b/Engine/Source/Programs/CSVTools/PerfReportTool/Telemetry.cs
--- a/Engine/Source/Programs/CSVTools/PerfReportTool/Telemetry.cs
+++ b/Engine/Source/Programs/CSVTools/PerfReportTool/Telemetry.cs
@@ -1,12 +1,13 @@
 // Copyright Epic Games, Inc. All Rights Reserved.
 
-using EpicGames.Core;
+using EpicGames.Analytics.PerformanceTrends;
 using EpicGames.Horde;
 using Microsoft.Extensions.DependencyInjection;
 using PerfSummaries;
 using System;
 using System.Collections;
 using System.Collections.Generic;
+using System.Linq;
 using System.Net.Http;
 using System.Threading.Tasks;
 
@@ -21,7 +22,8 @@ namespace PerfReportTool
 		private const string DefaultAppId = "Horde";
 		private const int DefaultAppVersion = 1;
 		private const string DefaultAppEnvironment = "Dev";
-		private const string EventName = "State.Performance.Debug";
+		private const string EventName = "State.Performance";
+		private const string DefaultSummaryTableName = "default";
 		private static ServiceProvider _serviceProvider = null;
 		private static IHordeClient _hordeClient = null;
 		private static HordeHttpClient _hordeHttpClient = null;
@@ -42,6 +44,7 @@ namespace PerfReportTool
 
 						if (defaultServer == null)
 						{
+							Console.WriteLine("Unable to construct HttpClient due to invalid default server.");
 							return null;
 						}
 
@@ -96,11 +99,43 @@ namespace PerfReportTool
 				_queuedRequests.Clear();
 			}
 
-			Task.WhenAll(tasks).GetAwaiter().GetResult();
+			try
+			{
+				Task.WhenAll(tasks).GetAwaiter().GetResult();
+			}
+			catch
+			{
+				Console.WriteLine("A task falted while awaiting all queued telemetry tasks.");
+			}
+
+			int succeeded = 0;
+			int failed = 0;
+
+			foreach (Task task in tasks)
+			{
+				if (task.IsCompletedSuccessfully)
+				{
+					succeeded++;
+				}
+				else if (task.IsFaulted || task.IsCanceled)
+				{
+					failed++;
+
+					if(task.Exception != null)
+					{
+						foreach(Exception ex in task.Exception.Flatten().InnerExceptions)
+						{
+							Console.WriteLine("Nested task exception: {0}", ex);
+						}
+					}
+				}
+			}
+
+			Console.WriteLine("Telemetry flush complete. Succeeded: {0}, Failed: {1}", succeeded, failed);
 		}
 
 		/// <summary>
-		/// Releases underyling singletons for the Telemetry helper.
+		/// Releases underlying singletons for the Telemetry helper.
 		/// </summary>
 		internal static void Dispose()
 		{
@@ -120,15 +155,18 @@ namespace PerfReportTool
 			}
 		}
 
+		#region -- SumamryTableRowData --
+
 		/// <summary>
-		/// Posts a row data's SummaryTableMetric to the Horde event stream, asynchronously.
+		/// Posts a rowData data's SummaryTableMetric to the Horde event stream, asynchronously.
 		/// </summary>
-		/// <param name="rowData">The row data to post.</param>
+		/// <param name="rowData">The rowData data to post.</param>
 		/// <returns>True if telemetry was emitted successfully, false otherwise.</returns>
 		internal async static Task<bool> PostSummaryTableMetricEventAsync(SummaryTableRowData rowData)
 		{
 			if (HordeClient == null)
 			{
+				Console.WriteLine("Skipping emission of rowData due to invalid HttpClient.");
 				return false;
 			}
 
@@ -144,37 +182,22 @@ namespace PerfReportTool
 				return false;
 			}
 
-			if (telemetryPayload == null)
-			{
-				return false;
-			}
-
-			try
-			{
-				using HttpResponseMessage response = await HordeHttpClient.PostTelemetryAsync(DefaultAppId, DefaultAppVersion, DefaultAppEnvironment, [telemetryPayload]);
-
-				if (response == null)
-				{
-					return false;
-				}
-
-				bool isSuccessStatusCode = response.IsSuccessStatusCode;
-				return isSuccessStatusCode;
-			}
-			catch (Exception ex)
-			{
-				Console.WriteLine("Encountered an exception while posting event. Exception: {0}", ex.Message);
-				return false;
-			}
+			return await EmitPayloadsAsync([telemetryPayload]);
 		}
 
 		/// <summary>
-		/// Enqueues a task to posts a row data's SummaryTableMetric to the Horde event stream.
+		/// Enqueues a task to posts a rowData data's SummaryTableMetric to the Horde event stream.
 		/// </summary>
-		/// <param name="rowData">The row data to post.</param>
+		/// <param name="rowData">The rowData data to post.</param>
 		/// <returns>A readonly list containing the queued tasks.</returns>
 		internal static IReadOnlyList<Task<bool>> EnqueuePostSummaryTableMetricEvent(SummaryTableRowData rowData)
 		{
+			if (HordeClient == null)
+			{
+				Console.WriteLine("Skipping emission of rowData due to invalid HttpClient.");
+				return _queuedRequests;
+			}
+
 			lock (_lock)
 			{
 				_queuedRequests.Add(PostSummaryTableMetricEventAsync(rowData));
@@ -184,14 +207,15 @@ namespace PerfReportTool
 		}
 
 		/// <summary>
-		/// Posts a row data's SummaryTableMetric to the Horde event stream.
+		/// Posts a rowData data's SummaryTableMetric to the Horde event stream.
 		/// </summary>
-		/// <param name="rowData">The row data to post.</param>
+		/// <param name="rowData">The rowData data to post.</param>
 		/// <returns>True if telemetry was emitted successfully, false otherwise.</returns>
 		internal static bool PostSummaryTableMetricEvent(SummaryTableRowData rowData)
 		{
 			if (HordeClient == null)
 			{
+				Console.WriteLine("Skipping emission of rowData due to invalid HttpClient.");
 				return false;
 			}
 
@@ -207,27 +231,15 @@ namespace PerfReportTool
 				return false;
 			}
 
-			if (telemetryPayload == null)
-			{
-				return false;
-			}
-
-			try
-			{
-				using HttpResponseMessage response = HordeHttpClient.PostTelemetryAsync(DefaultAppId, DefaultAppVersion, DefaultAppEnvironment, [telemetryPayload]).GetAwaiter().GetResult();
-				bool isSuccessStatusCode = response.IsSuccessStatusCode;
-				return isSuccessStatusCode;
-			}
-			catch (Exception ex)
-			{
-				Console.WriteLine("Encountered an exception while posting event. Exception: {0}", ex.Message);
-				return false;
-			}
+			return EmitPayloadsAsync([telemetryPayload]).GetAwaiter().GetResult();
 		}
 
 		private static Hashtable CreatePayload(SummaryTableRowData rowData)
 		{
 			Dictionary<string, dynamic> jsonDict = rowData.ToJsonDict(false, false);
+			string metadataMetricKey = SummaryTableElement.Type.CsvMetadata.ToString();
+			dynamic metadataMetric;
+
 			string summaryTableMetricKey = SummaryTableElement.Type.SummaryTableMetric.ToString();
 			dynamic summaryTableMetric;
 
@@ -241,19 +253,45 @@ namespace PerfReportTool
 				return null;
 			}
 
+			if (!jsonDict.TryGetValue(metadataMetricKey, out metadataMetric))
+			{
+				return null;
+			}
+
+			if (metadataMetric == null)
+			{
+				return null;
+			}
+
+			// Create resulting payload
+			// - Coalesce the metric data into the telemetry payload
+			// - Filter the telemetry payload of everything except the minimal metadata schema
+			// - Coalesce the summary items into the telemetry payload
+			// - Annotate the event as a global payload
 			Hashtable telemetryPayload = [];
-			telemetryPayload["EventName"] = EventName;
 
+			CoalesceDictionary(metadataMetric, telemetryPayload);
+
+			telemetryPayload = StandardTelemetryMetadata.CreateFilteredHashtable(telemetryPayload);
+
+			CoalesceDictionary(summaryTableMetric, telemetryPayload);
+
+			telemetryPayload["EventName"] = $"{EventName}.Global";
+			telemetryPayload = StandardTelemetryMetadata.SanitizeHashtable(telemetryPayload);
+			return telemetryPayload;
+		}
+
+		private static void CoalesceDictionary(dynamic summaryTableEntry, Hashtable telemetryPayload)
+		{
 			// Coalesce summary table metric values into telemetry payload via IDictionary || IDictionary<string,TValue> interfaces.
-			if (summaryTableMetric is IDictionary)
+			if (summaryTableEntry is IDictionary)
 			{
-				foreach (dynamic kvp in summaryTableMetric)
+				foreach (dynamic kvp in summaryTableEntry)
 				{
 					telemetryPayload[kvp.Key] = kvp.Value;
 				}
 			}
-
-			else if (summaryTableMetric is IDictionary<string, object> stringKeyGenericDict)
+			else if (summaryTableEntry is IDictionary<string, object> stringKeyGenericDict)
 			{
 				foreach (KeyValuePair<string, object> kvp in stringKeyGenericDict)
 				{
@@ -262,10 +300,86 @@ namespace PerfReportTool
 			}
 			else
 			{
-				throw new InvalidOperationException("SummaryTableMetric must be a dictionary or Dictionary<string, object>.");
+				throw new InvalidOperationException("SummaryTable entry must be a dictionary or Dictionary<string, object>.");
 			}
+		}
 
-			return telemetryPayload;
+		#endregion -- SummaryTableRowData --
+
+		#region -- SummaryTable --
+
+		/// <summary>
+		/// Enqueues a task to posts an entire summary table to the Horde event stream.
+		/// </summary>
+		/// <param name="summaryTableName">The name of the table.</param>
+		/// <param name="tableData">The table to post as an event.</param>
+		/// <param name="isCollated">If the table represents collated data or not.</param>
+		/// <returns>A readonly list containing the queued tasks.</returns>
+		internal static IReadOnlyList<Task<bool>> EnqueuePostSummaryTableMetricEvent(string summaryTableName, SummaryTable tableData, bool isCollated)
+		{
+			if (String.IsNullOrEmpty(summaryTableName))
+			{
+				summaryTableName = DefaultSummaryTableName;
+			}
+
+			IList<Hashtable> emissionPayloads = new List<Hashtable>(16);
+
+			foreach (IEnumerable<KeyValuePair<SummaryTableColumn, dynamic>> rowData in tableData.EnumerateTableDataByRow())
+			{
+				Hashtable telemetryPayload = [];
+
+				foreach (KeyValuePair<SummaryTableColumn, dynamic> cell in rowData)
+				{
+					SummaryTableColumn col = cell.Key;
+					dynamic val = cell.Value;
+
+					telemetryPayload[col.name] = val;
+				}
+
+				telemetryPayload["EventName"] = $"{EventName}.{summaryTableName}";
+				telemetryPayload["Collated"] = isCollated;
+
+				emissionPayloads.Add(telemetryPayload);
+			}
+
+			lock (_lock)
+			{
+				_queuedRequests.Add(EmitPayloadsAsync(emissionPayloads.ToArray()));
+			}
+
+			return _queuedRequests;
 		}
+
+		#endregion -- SummaryTable
+
+		#region -- Payload Emission --
+
+		private static async Task<bool> EmitPayloadsAsync(Hashtable[] telemetryPayloads)
+		{
+			if (telemetryPayloads == null || telemetryPayloads.Length == 0)
+			{
+				return false;
+			}
+
+			try
+			{
+				using HttpResponseMessage response = await HordeHttpClient.PostTelemetryAsync(DefaultAppId, DefaultAppVersion, DefaultAppEnvironment, telemetryPayloads);
+
+				if (response == null)
+				{
+					return false;
+				}
+
+				bool isSuccessStatusCode = response.IsSuccessStatusCode;
+				return isSuccessStatusCode;
+			}
+			catch (Exception ex)
+			{
+				Console.WriteLine("Encountered an exception while posting event. Exception: {0}", ex.Message);
+				return false;
+			}
+		}
+
+		#endregion -- Payload Emission --
 	}
 }
\ No newline at end of file
diff --git a/Engine/Source/Programs/Shared/EpicGames.Analytics.PerformanceTrends/EpicGames.Analytics.PerformanceTrends.csproj b/Engine/Source/Programs/Shared/EpicGames.Analytics.PerformanceTrends/EpicGames.Analytics.PerformanceTrends.csproj
new file mode 100644
--- /dev/null
+++ b/Engine/Source/Programs/Shared/EpicGames.Analytics.PerformanceTrends/EpicGames.Analytics.PerformanceTrends.csproj
@@ -0,0 +1,17 @@
+﻿<Project Sdk="Microsoft.NET.Sdk">
+  <PropertyGroup>
+    <TargetFramework>net8.0</TargetFramework>
+    <Deterministic>true</Deterministic>
+    <RootNamespace>EpicGames.Analytics</RootNamespace>
+    <GenerateDocumentationFile>True</GenerateDocumentationFile>
+    <AllowUnsafeBlocks>true</AllowUnsafeBlocks>
+  </PropertyGroup>
+
+  <PropertyGroup Condition="'$(Configuration)|$(Platform)' == 'Development|AnyCPU'">
+    <Optimize>false</Optimize>
+  </PropertyGroup>
+
+  <ItemGroup>
+    <ProjectReference Include="..\EpicGames.Analytics.Telemetry\EpicGames.Analytics.Telemetry.csproj" />
+  </ItemGroup>
+</Project>
diff --git a/Engine/Source/Programs/Shared/EpicGames.Analytics.PerformanceTrends/PerformanceTrendTelemetry.cs b/Engine/Source/Programs/Shared/EpicGames.Analytics.PerformanceTrends/PerformanceTrendTelemetry.cs
new file mode 100644
--- /dev/null
+++ b/Engine/Source/Programs/Shared/EpicGames.Analytics.PerformanceTrends/PerformanceTrendTelemetry.cs
@@ -0,0 +1,201 @@
+// Copyright Epic Games, Inc. All Rights Reserved.
+
+using System;
+using System.Collections;
+using System.ComponentModel.DataAnnotations.Schema;
+
+namespace EpicGames.Analytics.PerformanceTrends
+{
+	/// <summary>
+	/// Utility class to handle standard telemetry operations, and metadata processing.
+	/// </summary>
+	public static class StandardTelemetryMetadata
+	{
+		#region -- Internal Constants --
+
+		internal const string EngineVersion = "engineversion";
+		internal const string EngineReleaseVersion = "enginereleaseversion";
+		internal const string BuildVersion = "buildversion";
+		internal const string Platform = "platform";
+		internal const string TestId = "testid";
+		internal const string TestName = "testname";
+		internal const string TestConfigName = "testconfigname";
+		internal const string GauntletTestType = "gauntlettesttype";
+		internal const string GauntletSubTest = "gauntletsubtest";
+
+		#endregion -- Internal Constants --
+
+		#region -- Public Members --
+
+		/// <summary>
+		/// Sanitizes keys of the provided source hashtable.
+		/// </summary>
+		/// <param name="source">The source hashtable to sanitize.</param>
+		/// <returns>A copy of the hashtable with all keys sanitized.</returns>
+		public static Hashtable? SanitizeHashtable(Hashtable source)
+		{
+			if (source == null)
+			{
+				return null;
+			}
+
+			Hashtable sanitized = new Hashtable(source.Count);
+
+			foreach (DictionaryEntry entry in source)
+			{
+				if (entry.Key is string key)
+				{
+					string sanitizedKey = key.Replace("<", "_lt_", StringComparison.Ordinal).Replace(">", "_gt_", StringComparison.Ordinal);
+					sanitized[sanitizedKey] = entry.Value;
+				}
+				else
+				{
+					sanitized[entry.Key] = entry.Value;
+				}
+			}
+
+			return sanitized;
+		}
+
+		/// <summary>
+		/// The standard metadata keys.
+		/// </summary>
+		public static readonly string[] StandardMetadataKeys = new string[]
+		{
+				EngineVersion,
+				EngineReleaseVersion,
+				BuildVersion,
+				Platform,
+				TestId,
+				TestName,
+				TestConfigName,
+				GauntletTestType,
+				GauntletSubTest
+		};
+
+		/// <summary>
+		/// Creates a new hashtable only with the required metadata entries retained (if they exist).
+		/// </summary>
+		/// <param name="sourceTable">The source table.</param>
+		/// <returns>Hashtable with the required metadata entries.</returns>
+		public static Hashtable CreateFilteredHashtable(Hashtable sourceTable)
+		{
+			Hashtable returnTable = new Hashtable();
+
+			for (int i = 0; i < StandardMetadataKeys.Length; i++)
+			{
+				string key = StandardMetadataKeys[i];
+				if (sourceTable.ContainsKey(key))
+				{
+					returnTable[key] = sourceTable[key];
+				}
+			}
+
+			return returnTable;
+		}
+
+		#endregion -- Public Members --
+	}
+
+	/// <summary>
+	/// Abstract record that describes the shared metadata for performance trend telemetry.
+	/// </summary>
+	public abstract record PerformanceTrendTelemetry : AbstractTelemetryRecord
+	{
+		#region -- Public Members --
+
+		/// <summary>
+		/// The engine version of the telemetry.
+		/// </summary>
+		[Column("engine_version")]
+		public string? EngineVersion { get; init; }
+
+		/// <summary>
+		/// The engine release version of the telemetry.
+		/// </summary>
+		[Column("engine_release_version")]
+		public string? EngineReleaseVersion { get; init; }
+
+		/// <summary>
+		/// The build version of the telemetry.
+		/// </summary>
+		[Column("build_version")]
+		public string? BuildVersion { get; init; }
+
+		/// <summary>
+		/// The platform of the telemetry.
+		/// </summary>
+		[Column("platform")]
+		public string? Platform { get; init; }
+
+		/// <summary>
+		/// The test id for the telemetry, if applicable.
+		/// </summary>
+		[Column("test_id")]
+		public string? TestId { get; init; }
+
+		/// <summary>
+		/// The test name for the telemetry, if applicable.
+		/// </summary>
+		[Column("test_name")]
+		public string? TestName { get; init; }
+
+		/// <summary>
+		/// The test config name for the telemetry, if applicable.
+		/// </summary>
+		[Column("test_config_name")]
+		public string? TestConfigName { get; init; }
+
+		/// <summary>
+		/// The test type, if applicable.
+		/// </summary>
+		[Column("gauntlet_test_type")]
+		public string? GauntletTestType { get; init; }
+
+		/// <summary>
+		/// The sub test type, if applicable.
+		/// </summary>
+		[Column("gauntlet_sub_test")]
+		public string? GauntletSubTest { get; init; }
+
+		#endregion -- Public Members --
+
+		/// <summary>
+		/// Constructor.
+		/// </summary>
+		/// <param name="source">The source table to construct the telemetry record from.</param>
+		/// <param name="eventName">The event name.</param>
+		/// <param name="schemaVersion">The schema version of the telemetry.</param>
+		protected PerformanceTrendTelemetry(Hashtable source, string eventName, int schemaVersion) : base(eventName, schemaVersion)
+		{
+			if (source == null)
+			{
+				return;
+			}
+
+			EngineVersion = GetString(source, StandardTelemetryMetadata.EngineVersion);
+			EngineReleaseVersion = GetString(source, StandardTelemetryMetadata.EngineReleaseVersion);
+			BuildVersion = GetString(source, StandardTelemetryMetadata.BuildVersion);
+			Platform = GetString(source, StandardTelemetryMetadata.Platform);
+			TestId = GetString(source, StandardTelemetryMetadata.TestId);
+			TestName = GetString(source, StandardTelemetryMetadata.TestName);
+			TestConfigName = GetString(source, StandardTelemetryMetadata.TestConfigName);
+			GauntletTestType = GetString(source, StandardTelemetryMetadata.GauntletTestType);
+			GauntletSubTest = GetString(source, StandardTelemetryMetadata.GauntletSubTest);
+		}
+
+		#region -- Private API --
+
+		private static string? GetString(Hashtable table, string key)
+		{
+			if (!table.ContainsKey(key))
+			{
+				return null;
+			}
+
+			return table[key] as string;
+		}
+
+		#endregion -- Private API --
+	}
+}
\ No newline at end of file
