PRT - Don't map events before the start frame.
- Add extra image padding to the cropped image.

#rnx
#Synced-CL 49685785

diff --git a/Engine/Source/Programs/CSVTools/PerfReportTool/Summaries/MapOverlaySummary.cs b/Engine/Source/Programs/CSVTools/PerfReportTool/Summaries/MapOverlaySummary.cs
--- a/Engine/Source/Programs/CSVTools/PerfReportTool/Summaries/MapOverlaySummary.cs
+++ b/Engine/Source/Programs/CSVTools/PerfReportTool/Summaries/MapOverlaySummary.cs
@@ -196,9 +196,11 @@ namespace PerfSummaries
 			};
 		}
 
-		SVGPosition toSVGLocationClamped(float worldX, float worldY, Rectangle imageRegion)
+		SVGPosition toSVGLocationClamped(float worldX, float worldY, Rectangle imageRegion, int imagePadding)
 		{
 			SVGPosition svgPos = toSVGLocation(worldX, worldY, imageRegion);
+			svgPos.X += imagePadding;
+			svgPos.Y += imagePadding;
 			svgPos.X = Math.Min(Math.Max(0, svgPos.X), (int)imageWidth);
 			svgPos.Y = Math.Min(Math.Max(0, svgPos.Y), (int)imageHeight);
 			return svgPos;
@@ -315,8 +317,9 @@ namespace PerfSummaries
 				if (createCroppedImage)
 				{
 					// Generate Sub Image Location, Padded to give a small border.
-					SVGPosition minPos = toSVGLocationClamped(eventDimensions.minX, eventDimensions.minY, imageRegion);
-					SVGPosition maxPos = toSVGLocationClamped(eventDimensions.maxX, eventDimensions.maxY, imageRegion);
+					int extraImagePadding = 3;
+					SVGPosition minPos = toSVGLocationClamped(eventDimensions.minX, eventDimensions.minY, imageRegion, -extraImagePadding);
+					SVGPosition maxPos = toSVGLocationClamped(eventDimensions.maxX, eventDimensions.maxY, imageRegion, extraImagePadding);
 
 					if (minPos.X != maxPos.X && minPos.Y != maxPos.Y)
 					{
@@ -327,6 +330,8 @@ namespace PerfSummaries
 						int CenterY = minPos.Y + (SubRegionHeight / 2);
 						int UniformDim = Math.Max(SubRegionWidth, SubRegionHeight);
 						Rectangle subRegion = new Rectangle(CenterX - (UniformDim / 2), CenterY - (UniformDim / 2), UniformDim, UniformDim);
+						// Clamps us to the original image bounds in case we overflow it.
+						subRegion = Rectangle.Intersect(subRegion, imageRegion);
 
 						string croppedImageName = "crop_" + destImageFilename;
 						string outputTrimMapFilename = System.IO.Path.Combine(outputDirectory, croppedImageName);
@@ -401,7 +406,7 @@ namespace PerfSummaries
 				{
 					foreach (CsvEvent ev in csvStats.Events)
 					{
-						if (CsvStats.DoesSearchStringMatch(ev.Name, mapEvent.name))
+						if (ev.Frame >= startFrame && CsvStats.DoesSearchStringMatch(ev.Name, mapEvent.name))
 						{
 							frameEvents.Add(new KeyValuePair<int, MapOverlayEvent>(ev.Frame, mapEvent));
 						}
@@ -476,7 +481,7 @@ namespace PerfSummaries
 				{
 					foreach (CsvEvent ev in csvStats.Events)
 					{
-						if (CsvStats.DoesSearchStringMatch(ev.Name, mapEvent.name))
+						if (ev.Frame >= startFrame && CsvStats.DoesSearchStringMatch(ev.Name, mapEvent.name))
 						{
 							string eventText = mapEvent.shortName != null ? mapEvent.shortName : ev.Name;
 							float x = xStat.samples[ev.Frame];
