[PerformanceReportTool][Analytics] Forwarding through Session Id, added option to include Session Label. Also included some base validation via EpicGames.Analytics.Telemetry tooling, and collated table naming.

#Synced-CL 49930492

diff --git a/Engine/Source/Programs/CSVTools/PerfReportTool/PerfReportTool.cs b/Engine/Source/Programs/CSVTools/PerfReportTool/PerfReportTool.cs
--- a/Engine/Source/Programs/CSVTools/PerfReportTool/PerfReportTool.cs
+++ b/Engine/Source/Programs/CSVTools/PerfReportTool/PerfReportTool.cs
@@ -24,7 +24,7 @@ namespace PerfReportTool
 	class Version
 	{
 		// Format: Major.Minor.Bugfix
-		private static string VersionString = "4.265.0";
+		private static string VersionString = "4.266.0";
 
 		public static string Get() { return VersionString; }
 	};
@@ -216,6 +216,8 @@ namespace PerfReportTool
 			"  -summaryTableDataToTelemetry : publish the summary table as a telemetry event.\n" +
 			"  -summaryTableDataToTelemetryVerbose : publish the verbose form of every row data (without summary table filter).\n" +
 			"     Will also emit summary table data.\n" +
+			"  -telemetrySessionLabel <label>: \n" +
+			"     Optional label to include on all telemetry emissions.\n" +
 			"\n" +
 			"Performance args for bulk mode:\n" +
 			"  -precacheCount <n> : number of CSV files to precache in the lookahead cache (0 for no precache)\n" +
@@ -455,6 +457,7 @@ namespace PerfReportTool
 
 			bool summaryTableDataToTelemetryVerbose = GetBoolArg("summaryTableDataToTelemetryVerbose");
 			bool summaryTableDataToTelemetry = GetBoolArg("summaryTableDataToTelemetry") || summaryTableDataToTelemetryVerbose;
+			string telemetrySessionLabel = GetArg("telemetrySessionLabel", null);
 
 			if (summaryJsonOutPath != null)
 			{
@@ -735,7 +738,7 @@ namespace PerfReportTool
 							// Emit telemetry if enabled
 							if (summaryTableDataToTelemetryVerbose)
 							{
-								Telemetry.EnqueuePostSummaryTableMetricEvent(rowData);
+								Telemetry.EnqueuePostSummaryTableMetricEvent(rowData, telemetrySessionLabel);
 							}
 
 							perfLog.LogTiming("  AddRowData");
@@ -808,11 +811,11 @@ namespace PerfReportTool
 				SummaryTableInfo tableInfo = GetTableInfo(summaryTableName, customSummaryTableFilter, customSummaryTableRowSort);
 				if (!bCollateTableOnly)
 				{
-					WriteSummaryTableReport(outputDir, summaryTableFilename, summaryTable, tableInfo, false, bCsvTable, bSpreadsheetFriendlyStrings, null);
+					WriteSummaryTableReport(outputDir, summaryTableFilename, summaryTable, tableInfo, false, bCsvTable, bSpreadsheetFriendlyStrings, null, telemetrySessionLabel);
 				}
 				if (bCollateTable)
 				{
-					WriteSummaryTableReport(outputDir, collatedTableFilename, summaryTable, tableInfo, true, bCsvTable, bSpreadsheetFriendlyStrings, weightByColumnName);
+					WriteSummaryTableReport(outputDir, collatedTableFilename, summaryTable, tableInfo, true, bCsvTable, bSpreadsheetFriendlyStrings, weightByColumnName, telemetrySessionLabel);
 				}
 
 				// EmailTable is hardcoded to use the condensed type
@@ -886,7 +889,7 @@ namespace PerfReportTool
 			WriteSummaryTableReport(outputDir, filenameWithoutExtension, table, tableInfo, bCollated, bToCSV, bSpreadsheetFriendlyStrings, weightByColumnName);
 		}
 
-		void WriteSummaryTableReport(string outputDir, string filenameWithoutExtension, SummaryTable table, SummaryTableInfo tableInfo, bool bCollated, bool bToCSV, bool bSpreadsheetFriendlyStrings, string weightByColumnNameOverride)
+		void WriteSummaryTableReport(string outputDir, string filenameWithoutExtension, SummaryTable table, SummaryTableInfo tableInfo, bool bCollated, bool bToCSV, bool bSpreadsheetFriendlyStrings, string weightByColumnNameOverride, string telemetrySessionLabel = null)
 		{
 			string weightByColumnName = weightByColumnNameOverride != null ? weightByColumnNameOverride : tableInfo.weightByColumn;
 			if (GetBoolArg("noWeightedAvg"))
@@ -997,7 +1000,7 @@ namespace PerfReportTool
 				}
 
 				string summaryTableName = GetArg("summaryTable", null);
-				Telemetry.EnqueuePostSummaryTableMetricEvent(summaryTableName, reconstitutedTable, bCollated);
+				Telemetry.EnqueuePostSummaryTableMetricEvent(summaryTableName, reconstitutedTable, bCollated, telemetrySessionLabel);
 			}
 		}
 
diff --git a/Engine/Source/Programs/CSVTools/PerfReportTool/Telemetry.cs b/Engine/Source/Programs/CSVTools/PerfReportTool/Telemetry.cs
--- a/Engine/Source/Programs/CSVTools/PerfReportTool/Telemetry.cs
+++ b/Engine/Source/Programs/CSVTools/PerfReportTool/Telemetry.cs
@@ -1,7 +1,9 @@
 // Copyright Epic Games, Inc. All Rights Reserved.
 
+using EpicGames.Analytics;
 using EpicGames.Analytics.PerformanceTrends;
 using EpicGames.Horde;
+using EpicGames.Horde.Telemetry;
 using Microsoft.Extensions.DependencyInjection;
 using PerfSummaries;
 using System;
@@ -9,6 +11,7 @@ using System.Collections;
 using System.Collections.Generic;
 using System.Linq;
 using System.Net.Http;
+using System.Threading;
 using System.Threading.Tasks;
 
 namespace PerfReportTool
@@ -24,10 +27,13 @@ namespace PerfReportTool
 		private const string DefaultAppEnvironment = "Dev";
 		private const string EventName = "State.Performance";
 		private const string DefaultSummaryTableName = "default";
+		private const int CurrentSchemaVersion = 1;
+		private const int DefaultTimeoutSeconds = 15;
 		private static ServiceProvider _serviceProvider = null;
 		private static IHordeClient _hordeClient = null;
 		private static HordeHttpClient _hordeHttpClient = null;
 		private readonly static List<Task<bool>> _queuedRequests = new List<Task<bool>>(16);
+		private readonly static string _sessionId = Guid.NewGuid().ToString("N");
 
 		/// <summary>
 		/// Gets the HordeClient used to create Horde Http Clients.
@@ -44,7 +50,7 @@ namespace PerfReportTool
 
 						if (defaultServer == null)
 						{
-							Console.WriteLine("Unable to construct HttpClient due to invalid default server.");
+							Console.WriteLine("[Telemetry] Unable to construct HttpClient due to invalid default server.");
 							return null;
 						}
 
@@ -85,7 +91,7 @@ namespace PerfReportTool
 		/// </summary>
 		internal static void Flush()
 		{
-			Console.WriteLine("Flushing {0} telemetry events.", _queuedRequests.Count);
+			Console.WriteLine("[Telemetry] Flushing {0} telemetry events.", _queuedRequests.Count);
 
 			if (_queuedRequests.Count == 0)
 			{
@@ -105,7 +111,7 @@ namespace PerfReportTool
 			}
 			catch
 			{
-				Console.WriteLine("A task falted while awaiting all queued telemetry tasks.");
+				Console.WriteLine("[Telemetry] A task falted while awaiting all queued telemetry tasks.");
 			}
 
 			int succeeded = 0;
@@ -121,17 +127,17 @@ namespace PerfReportTool
 				{
 					failed++;
 
-					if(task.Exception != null)
+					if (task.Exception != null)
 					{
-						foreach(Exception ex in task.Exception.Flatten().InnerExceptions)
+						foreach (Exception ex in task.Exception.Flatten().InnerExceptions)
 						{
-							Console.WriteLine("Nested task exception: {0}", ex);
+							Console.WriteLine("[Telemetry] Nested task exception: {0}", ex);
 						}
 					}
 				}
 			}
 
-			Console.WriteLine("Telemetry flush complete. Succeeded: {0}, Failed: {1}", succeeded, failed);
+			Console.WriteLine("[Telemetry] Telemetry flush complete. Succeeded: {0}, Failed: {1}", succeeded, failed);
 		}
 
 		/// <summary>
@@ -162,11 +168,11 @@ namespace PerfReportTool
 		/// </summary>
 		/// <param name="rowData">The rowData data to post.</param>
 		/// <returns>True if telemetry was emitted successfully, false otherwise.</returns>
-		internal async static Task<bool> PostSummaryTableMetricEventAsync(SummaryTableRowData rowData)
+		internal async static Task<bool> PostSummaryTableMetricEventAsync(SummaryTableRowData rowData, string sessionLabel = null)
 		{
 			if (HordeClient == null)
 			{
-				Console.WriteLine("Skipping emission of rowData due to invalid HttpClient.");
+				Console.WriteLine("[Telemetry] Skipping emission of rowData due to invalid HttpClient.");
 				return false;
 			}
 
@@ -178,29 +184,30 @@ namespace PerfReportTool
 			}
 			catch (InvalidOperationException ex)
 			{
-				Console.WriteLine("Skipping emission of rowData due to invalid payload construction. Exception: {0}", ex.Message);
+				Console.WriteLine("[Telemetry] Skipping emission of rowData due to invalid payload construction. Exception: {0}", ex.Message);
 				return false;
 			}
 
-			return await EmitPayloadsAsync([telemetryPayload]);
+			return await EmitPayloadsAsync([telemetryPayload], sessionLabel);
 		}
 
 		/// <summary>
 		/// Enqueues a task to posts a rowData data's SummaryTableMetric to the Horde event stream.
 		/// </summary>
 		/// <param name="rowData">The rowData data to post.</param>
+		/// <param name="sessionLabel">Optional session label used to correlate telemetry events.</param>
 		/// <returns>A readonly list containing the queued tasks.</returns>
-		internal static IReadOnlyList<Task<bool>> EnqueuePostSummaryTableMetricEvent(SummaryTableRowData rowData)
+		internal static IReadOnlyList<Task<bool>> EnqueuePostSummaryTableMetricEvent(SummaryTableRowData rowData, string sessionLabel = null)
 		{
 			if (HordeClient == null)
 			{
-				Console.WriteLine("Skipping emission of rowData due to invalid HttpClient.");
+				Console.WriteLine("[Telemetry] Skipping emission of rowData due to invalid HttpClient.");
 				return _queuedRequests;
 			}
 
 			lock (_lock)
 			{
-				_queuedRequests.Add(PostSummaryTableMetricEventAsync(rowData));
+				_queuedRequests.Add(PostSummaryTableMetricEventAsync(rowData, sessionLabel));
 			}
 
 			return _queuedRequests;
@@ -210,12 +217,13 @@ namespace PerfReportTool
 		/// Posts a rowData data's SummaryTableMetric to the Horde event stream.
 		/// </summary>
 		/// <param name="rowData">The rowData data to post.</param>
+		/// <param name="sessionLabel">Optional session label used to correlate telemetry events.</param>
 		/// <returns>True if telemetry was emitted successfully, false otherwise.</returns>
-		internal static bool PostSummaryTableMetricEvent(SummaryTableRowData rowData)
+		internal static bool PostSummaryTableMetricEvent(SummaryTableRowData rowData, string sessionLabel = null)
 		{
 			if (HordeClient == null)
 			{
-				Console.WriteLine("Skipping emission of rowData due to invalid HttpClient.");
+				Console.WriteLine("[Telemetry] Skipping emission of rowData due to invalid HttpClient.");
 				return false;
 			}
 
@@ -227,14 +235,14 @@ namespace PerfReportTool
 			}
 			catch (InvalidOperationException ex)
 			{
-				Console.WriteLine("Skipping emission of rowData due to invalid payload construction. Exception: {0}", ex.Message);
+				Console.WriteLine("[Telemetry] Skipping emission of rowData due to invalid payload construction. Exception: {0}", ex.Message);
 				return false;
 			}
 
-			return EmitPayloadsAsync([telemetryPayload]).GetAwaiter().GetResult();
+			return EmitPayloadsAsync([telemetryPayload], sessionLabel).GetAwaiter().GetResult();
 		}
 
-		private static Hashtable CreatePayload(SummaryTableRowData rowData)
+		private static Hashtable CreatePayload(SummaryTableRowData rowData, string sessionLabel = null)
 		{
 			Dictionary<string, dynamic> jsonDict = rowData.ToJsonDict(false, false);
 			string metadataMetricKey = SummaryTableElement.Type.CsvMetadata.ToString();
@@ -276,8 +284,11 @@ namespace PerfReportTool
 
 			CoalesceDictionary(summaryTableMetric, telemetryPayload);
 
-			telemetryPayload["EventName"] = $"{EventName}.Global";
 			telemetryPayload = StandardTelemetryMetadata.SanitizeHashtable(telemetryPayload);
+
+			HordeContextTelemetryRecord telemetryRecord = new HordeContextTelemetryRecord($"{EventName}.Global", CurrentSchemaVersion, _sessionId, sessionLabel);
+			telemetryPayload = AnalyticsTelemetryHelpers.MergeLeftCopy(telemetryPayload, telemetryRecord.ToHashtable());
+
 			return telemetryPayload;
 		}
 
@@ -300,7 +311,7 @@ namespace PerfReportTool
 			}
 			else
 			{
-				throw new InvalidOperationException("SummaryTable entry must be a dictionary or Dictionary<string, object>.");
+				throw new InvalidOperationException("[Telemetry] SummaryTable entry must be a dictionary or Dictionary<string, object>.");
 			}
 		}
 
@@ -314,8 +325,9 @@ namespace PerfReportTool
 		/// <param name="summaryTableName">The name of the table.</param>
 		/// <param name="tableData">The table to post as an event.</param>
 		/// <param name="isCollated">If the table represents collated data or not.</param>
+		/// <param name="sessionLabel">Optional session label used to correlate telemetry events.</param>
 		/// <returns>A readonly list containing the queued tasks.</returns>
-		internal static IReadOnlyList<Task<bool>> EnqueuePostSummaryTableMetricEvent(string summaryTableName, SummaryTable tableData, bool isCollated)
+		internal static IReadOnlyList<Task<bool>> EnqueuePostSummaryTableMetricEvent(string summaryTableName, SummaryTable tableData, bool isCollated, string sessionLabel = null)
 		{
 			if (String.IsNullOrEmpty(summaryTableName))
 			{
@@ -323,6 +335,9 @@ namespace PerfReportTool
 			}
 
 			IList<Hashtable> emissionPayloads = new List<Hashtable>(16);
+			string eventName = isCollated ? $"{EventName}.{summaryTableName}.Collated" : $"{EventName}.{summaryTableName}";
+			HordeContextTelemetryRecord telemetryRecord = new HordeContextTelemetryRecord(eventName, CurrentSchemaVersion, _sessionId, sessionLabel);
+			Hashtable baseContext = telemetryRecord.ToHashtable();
 
 			foreach (IEnumerable<KeyValuePair<SummaryTableColumn, dynamic>> rowData in tableData.EnumerateTableDataByRow())
 			{
@@ -336,15 +351,18 @@ namespace PerfReportTool
 					telemetryPayload[col.name] = val;
 				}
 
-				telemetryPayload["EventName"] = $"{EventName}.{summaryTableName}";
+				telemetryPayload["SummaryName"] = summaryTableName;
 				telemetryPayload["Collated"] = isCollated;
 
+				telemetryPayload = StandardTelemetryMetadata.SanitizeHashtable(telemetryPayload);
+				telemetryPayload = AnalyticsTelemetryHelpers.MergeLeft(telemetryPayload, baseContext);
+
 				emissionPayloads.Add(telemetryPayload);
 			}
 
 			lock (_lock)
 			{
-				_queuedRequests.Add(EmitPayloadsAsync(emissionPayloads.ToArray()));
+				_queuedRequests.Add(EmitPayloadsAsync(emissionPayloads.ToArray(), sessionLabel));
 			}
 
 			return _queuedRequests;
@@ -354,7 +372,7 @@ namespace PerfReportTool
 
 		#region -- Payload Emission --
 
-		private static async Task<bool> EmitPayloadsAsync(Hashtable[] telemetryPayloads)
+		private static async Task<bool> EmitPayloadsAsync(Hashtable[] telemetryPayloads, string sessionLabel = null)
 		{
 			if (telemetryPayloads == null || telemetryPayloads.Length == 0)
 			{
@@ -363,7 +381,22 @@ namespace PerfReportTool
 
 			try
 			{
-				using HttpResponseMessage response = await HordeHttpClient.PostTelemetryAsync(DefaultAppId, DefaultAppVersion, DefaultAppEnvironment, telemetryPayloads);
+				using CancellationTokenSource cancellationToken = new CancellationTokenSource(TimeSpan.FromSeconds(DefaultTimeoutSeconds));
+				List<string> missingItems = new List<string>();
+				for (int i = 0; i < telemetryPayloads.Count(); ++i)
+				{
+					ColumnValidator<KeyStatTelemetryRecord>.HasCompletePropertySet(telemetryPayloads[i], true, null);
+				}
+
+				if (missingItems.Count > 0)
+				{
+					for (int i = 0; i < missingItems.Count; ++i)
+					{
+						Console.WriteLine("[Telemetry] Property was not found in telemetry payload: {0}", missingItems[i]);
+					}
+				}
+
+				using HttpResponseMessage response = await HordeHttpClient.PostTelemetryAsync(DefaultAppId, DefaultAppVersion, DefaultAppEnvironment, telemetryPayloads, _sessionId, cancellationToken.Token);
 
 				if (response == null)
 				{
@@ -375,7 +408,7 @@ namespace PerfReportTool
 			}
 			catch (Exception ex)
 			{
-				Console.WriteLine("Encountered an exception while posting event. Exception: {0}", ex.Message);
+				Console.WriteLine("[Telemetry] Encountered an exception while posting event. Exception: {0}", ex.Message);
 				return false;
 			}
 		}
