[PerformanceTool][Analytics] Adding telemetryDataSource for fundamental attribution of csvbin.

#Synced-CL 49930492

diff --git a/Engine/Source/Programs/CSVTools/PerfReportTool/PerfReportTool.cs b/Engine/Source/Programs/CSVTools/PerfReportTool/PerfReportTool.cs
--- a/Engine/Source/Programs/CSVTools/PerfReportTool/PerfReportTool.cs
+++ b/Engine/Source/Programs/CSVTools/PerfReportTool/PerfReportTool.cs
@@ -24,7 +24,7 @@ namespace PerfReportTool
 	class Version
 	{
 		// Format: Major.Minor.Bugfix
-		private static string VersionString = "4.266.0";
+		private static string VersionString = "4.267.0";
 
 		public static string Get() { return VersionString; }
 	};
@@ -218,6 +218,8 @@ namespace PerfReportTool
 			"     Will also emit summary table data.\n" +
 			"  -telemetrySessionLabel <label>: \n" +
 			"     Optional label to include on all telemetry emissions.\n" +
+			"  -telemetryDataSource <dataSource>: \n" +
+			"     Optional data source to include on all telemetry emissions.\n" +
 			"\n" +
 			"Performance args for bulk mode:\n" +
 			"  -precacheCount <n> : number of CSV files to precache in the lookahead cache (0 for no precache)\n" +
@@ -458,6 +460,7 @@ namespace PerfReportTool
 			bool summaryTableDataToTelemetryVerbose = GetBoolArg("summaryTableDataToTelemetryVerbose");
 			bool summaryTableDataToTelemetry = GetBoolArg("summaryTableDataToTelemetry") || summaryTableDataToTelemetryVerbose;
 			string telemetrySessionLabel = GetArg("telemetrySessionLabel", null);
+			string telemetryDataSource = GetArg("telemetryDataSource", null);
 
 			if (summaryJsonOutPath != null)
 			{
@@ -738,7 +741,7 @@ namespace PerfReportTool
 							// Emit telemetry if enabled
 							if (summaryTableDataToTelemetryVerbose)
 							{
-								Telemetry.EnqueuePostSummaryTableMetricEvent(rowData, telemetrySessionLabel);
+								Telemetry.EnqueuePostSummaryTableMetricEvent(rowData, telemetrySessionLabel, telemetryDataSource);
 							}
 
 							perfLog.LogTiming("  AddRowData");
@@ -811,11 +814,11 @@ namespace PerfReportTool
 				SummaryTableInfo tableInfo = GetTableInfo(summaryTableName, customSummaryTableFilter, customSummaryTableRowSort);
 				if (!bCollateTableOnly)
 				{
-					WriteSummaryTableReport(outputDir, summaryTableFilename, summaryTable, tableInfo, false, bCsvTable, bSpreadsheetFriendlyStrings, null, telemetrySessionLabel);
+					WriteSummaryTableReport(outputDir, summaryTableFilename, summaryTable, tableInfo, false, bCsvTable, bSpreadsheetFriendlyStrings, null, telemetrySessionLabel, telemetryDataSource);
 				}
 				if (bCollateTable)
 				{
-					WriteSummaryTableReport(outputDir, collatedTableFilename, summaryTable, tableInfo, true, bCsvTable, bSpreadsheetFriendlyStrings, weightByColumnName, telemetrySessionLabel);
+					WriteSummaryTableReport(outputDir, collatedTableFilename, summaryTable, tableInfo, true, bCsvTable, bSpreadsheetFriendlyStrings, weightByColumnName, telemetrySessionLabel, telemetryDataSource);
 				}
 
 				// EmailTable is hardcoded to use the condensed type
@@ -889,7 +892,7 @@ namespace PerfReportTool
 			WriteSummaryTableReport(outputDir, filenameWithoutExtension, table, tableInfo, bCollated, bToCSV, bSpreadsheetFriendlyStrings, weightByColumnName);
 		}
 
-		void WriteSummaryTableReport(string outputDir, string filenameWithoutExtension, SummaryTable table, SummaryTableInfo tableInfo, bool bCollated, bool bToCSV, bool bSpreadsheetFriendlyStrings, string weightByColumnNameOverride, string telemetrySessionLabel = null)
+		void WriteSummaryTableReport(string outputDir, string filenameWithoutExtension, SummaryTable table, SummaryTableInfo tableInfo, bool bCollated, bool bToCSV, bool bSpreadsheetFriendlyStrings, string weightByColumnNameOverride, string telemetrySessionLabel = null, string telemetryDataSource = null)
 		{
 			string weightByColumnName = weightByColumnNameOverride != null ? weightByColumnNameOverride : tableInfo.weightByColumn;
 			if (GetBoolArg("noWeightedAvg"))
@@ -1000,7 +1003,7 @@ namespace PerfReportTool
 				}
 
 				string summaryTableName = GetArg("summaryTable", null);
-				Telemetry.EnqueuePostSummaryTableMetricEvent(summaryTableName, reconstitutedTable, bCollated, telemetrySessionLabel);
+				Telemetry.EnqueuePostSummaryTableMetricEvent(summaryTableName, reconstitutedTable, bCollated, telemetrySessionLabel, telemetryDataSource);
 			}
 		}
 
diff --git a/Engine/Source/Programs/CSVTools/PerfReportTool/Telemetry.cs b/Engine/Source/Programs/CSVTools/PerfReportTool/Telemetry.cs
--- a/Engine/Source/Programs/CSVTools/PerfReportTool/Telemetry.cs
+++ b/Engine/Source/Programs/CSVTools/PerfReportTool/Telemetry.cs
@@ -167,8 +167,10 @@ namespace PerfReportTool
 		/// Posts a rowData data's SummaryTableMetric to the Horde event stream, asynchronously.
 		/// </summary>
 		/// <param name="rowData">The rowData data to post.</param>
+		/// <param name="sessionLabel">Optional session label used to correlate telemetry events.</param>
+		/// <param name="dataSource">Optional data source used to group telemetry events.</param>
 		/// <returns>True if telemetry was emitted successfully, false otherwise.</returns>
-		internal async static Task<bool> PostSummaryTableMetricEventAsync(SummaryTableRowData rowData, string sessionLabel = null)
+		internal async static Task<bool> PostSummaryTableMetricEventAsync(SummaryTableRowData rowData, string sessionLabel = null, string dataSource = null)
 		{
 			if (HordeClient == null)
 			{
@@ -180,7 +182,7 @@ namespace PerfReportTool
 
 			try
 			{
-				telemetryPayload = CreatePayload(rowData);
+				telemetryPayload = CreatePayload(rowData, sessionLabel, dataSource);
 			}
 			catch (InvalidOperationException ex)
 			{
@@ -188,7 +190,7 @@ namespace PerfReportTool
 				return false;
 			}
 
-			return await EmitPayloadsAsync([telemetryPayload], sessionLabel);
+			return await EmitPayloadsAsync([telemetryPayload]);
 		}
 
 		/// <summary>
@@ -196,8 +198,9 @@ namespace PerfReportTool
 		/// </summary>
 		/// <param name="rowData">The rowData data to post.</param>
 		/// <param name="sessionLabel">Optional session label used to correlate telemetry events.</param>
+		/// <param name="dataSource">Optional data source used to group telemetry events.</param>
 		/// <returns>A readonly list containing the queued tasks.</returns>
-		internal static IReadOnlyList<Task<bool>> EnqueuePostSummaryTableMetricEvent(SummaryTableRowData rowData, string sessionLabel = null)
+		internal static IReadOnlyList<Task<bool>> EnqueuePostSummaryTableMetricEvent(SummaryTableRowData rowData, string sessionLabel = null, string dataSource = null)
 		{
 			if (HordeClient == null)
 			{
@@ -207,7 +210,7 @@ namespace PerfReportTool
 
 			lock (_lock)
 			{
-				_queuedRequests.Add(PostSummaryTableMetricEventAsync(rowData, sessionLabel));
+				_queuedRequests.Add(PostSummaryTableMetricEventAsync(rowData, sessionLabel, dataSource));
 			}
 
 			return _queuedRequests;
@@ -218,8 +221,9 @@ namespace PerfReportTool
 		/// </summary>
 		/// <param name="rowData">The rowData data to post.</param>
 		/// <param name="sessionLabel">Optional session label used to correlate telemetry events.</param>
+		/// <param name="dataSource">Optional data source used to group telemetry events.</param>
 		/// <returns>True if telemetry was emitted successfully, false otherwise.</returns>
-		internal static bool PostSummaryTableMetricEvent(SummaryTableRowData rowData, string sessionLabel = null)
+		internal static bool PostSummaryTableMetricEvent(SummaryTableRowData rowData, string sessionLabel = null, string dataSource = null)
 		{
 			if (HordeClient == null)
 			{
@@ -231,7 +235,7 @@ namespace PerfReportTool
 
 			try
 			{
-				telemetryPayload = CreatePayload(rowData);
+				telemetryPayload = CreatePayload(rowData, sessionLabel, dataSource);
 			}
 			catch (InvalidOperationException ex)
 			{
@@ -239,10 +243,10 @@ namespace PerfReportTool
 				return false;
 			}
 
-			return EmitPayloadsAsync([telemetryPayload], sessionLabel).GetAwaiter().GetResult();
+			return EmitPayloadsAsync([telemetryPayload]).GetAwaiter().GetResult();
 		}
 
-		private static Hashtable CreatePayload(SummaryTableRowData rowData, string sessionLabel = null)
+		private static Hashtable CreatePayload(SummaryTableRowData rowData, string sessionLabel = null, string dataSource = null)
 		{
 			Dictionary<string, dynamic> jsonDict = rowData.ToJsonDict(false, false);
 			string metadataMetricKey = SummaryTableElement.Type.CsvMetadata.ToString();
@@ -286,6 +290,11 @@ namespace PerfReportTool
 
 			telemetryPayload = StandardTelemetryMetadata.SanitizeHashtable(telemetryPayload);
 
+			if (!telemetryPayload.ContainsKey("DataSource") && !String.IsNullOrEmpty(dataSource))
+			{
+				telemetryPayload["DataSource"] = dataSource;
+			}
+
 			HordeContextTelemetryRecord telemetryRecord = new HordeContextTelemetryRecord($"{EventName}.Global", CurrentSchemaVersion, _sessionId, sessionLabel);
 			telemetryPayload = AnalyticsTelemetryHelpers.MergeLeftCopy(telemetryPayload, telemetryRecord.ToHashtable());
 
@@ -326,8 +335,9 @@ namespace PerfReportTool
 		/// <param name="tableData">The table to post as an event.</param>
 		/// <param name="isCollated">If the table represents collated data or not.</param>
 		/// <param name="sessionLabel">Optional session label used to correlate telemetry events.</param>
+		/// <param name="dataSource">Optional data source used to group telemetry events.</param>
 		/// <returns>A readonly list containing the queued tasks.</returns>
-		internal static IReadOnlyList<Task<bool>> EnqueuePostSummaryTableMetricEvent(string summaryTableName, SummaryTable tableData, bool isCollated, string sessionLabel = null)
+		internal static IReadOnlyList<Task<bool>> EnqueuePostSummaryTableMetricEvent(string summaryTableName, SummaryTable tableData, bool isCollated, string sessionLabel = null, string dataSource = null)
 		{
 			if (String.IsNullOrEmpty(summaryTableName))
 			{
@@ -352,6 +362,12 @@ namespace PerfReportTool
 				}
 
 				telemetryPayload["SummaryName"] = summaryTableName;
+
+				if (!telemetryPayload.ContainsKey("DataSource") && !String.IsNullOrEmpty(dataSource))
+				{
+					telemetryPayload["DataSource"] = dataSource;
+				}
+
 				telemetryPayload["Collated"] = isCollated;
 
 				telemetryPayload = StandardTelemetryMetadata.SanitizeHashtable(telemetryPayload);
@@ -362,7 +378,7 @@ namespace PerfReportTool
 
 			lock (_lock)
 			{
-				_queuedRequests.Add(EmitPayloadsAsync(emissionPayloads.ToArray(), sessionLabel));
+				_queuedRequests.Add(EmitPayloadsAsync(emissionPayloads.ToArray()));
 			}
 
 			return _queuedRequests;
@@ -372,7 +388,7 @@ namespace PerfReportTool
 
 		#region -- Payload Emission --
 
-		private static async Task<bool> EmitPayloadsAsync(Hashtable[] telemetryPayloads, string sessionLabel = null)
+		private static async Task<bool> EmitPayloadsAsync(Hashtable[] telemetryPayloads)
 		{
 			if (telemetryPayloads == null || telemetryPayloads.Length == 0)
 			{
