[Horde][Analytics] Expanding base context to include some hardware capabilities, and replacing DataSource with new "TestIdentity" nomenclature.

#Synced-CL 50425380

diff --git a/Engine/Source/Programs/CSVTools/PerfReportTool/PerfReportTool.cs b/Engine/Source/Programs/CSVTools/PerfReportTool/PerfReportTool.cs
--- a/Engine/Source/Programs/CSVTools/PerfReportTool/PerfReportTool.cs
+++ b/Engine/Source/Programs/CSVTools/PerfReportTool/PerfReportTool.cs
@@ -24,7 +24,7 @@ namespace PerfReportTool
 	class Version
 	{
 		// Format: Major.Minor.Bugfix
-		private static string VersionString = "4.267.0";
+		private static string VersionString = "4.268.0";
 
 		public static string Get() { return VersionString; }
 	};
@@ -218,7 +218,7 @@ namespace PerfReportTool
 			"     Will also emit summary table data.\n" +
 			"  -telemetrySessionLabel <label>: \n" +
 			"     Optional label to include on all telemetry emissions.\n" +
-			"  -telemetryDataSource <dataSource>: \n" +
+			"  -telemetryTestIdentity <testIdentity>: \n" +
 			"     Optional data source to include on all telemetry emissions.\n" +
 			"\n" +
 			"Performance args for bulk mode:\n" +
@@ -460,7 +460,7 @@ namespace PerfReportTool
 			bool summaryTableDataToTelemetryVerbose = GetBoolArg("summaryTableDataToTelemetryVerbose");
 			bool summaryTableDataToTelemetry = GetBoolArg("summaryTableDataToTelemetry") || summaryTableDataToTelemetryVerbose;
 			string telemetrySessionLabel = GetArg("telemetrySessionLabel", null);
-			string telemetryDataSource = GetArg("telemetryDataSource", null);
+			string telemetryTestIdentity = GetArg("telemetryTestIdentity", null);
 
 			if (summaryJsonOutPath != null)
 			{
@@ -741,7 +741,7 @@ namespace PerfReportTool
 							// Emit telemetry if enabled
 							if (summaryTableDataToTelemetryVerbose)
 							{
-								Telemetry.EnqueuePostSummaryTableMetricEvent(rowData, telemetrySessionLabel, telemetryDataSource);
+								Telemetry.EnqueuePostSummaryTableMetricEvent(rowData, telemetrySessionLabel, telemetryTestIdentity);
 							}
 
 							perfLog.LogTiming("  AddRowData");
@@ -814,11 +814,11 @@ namespace PerfReportTool
 				SummaryTableInfo tableInfo = GetTableInfo(summaryTableName, customSummaryTableFilter, customSummaryTableRowSort);
 				if (!bCollateTableOnly)
 				{
-					WriteSummaryTableReport(outputDir, summaryTableFilename, summaryTable, tableInfo, false, bCsvTable, bSpreadsheetFriendlyStrings, null, telemetrySessionLabel, telemetryDataSource);
+					WriteSummaryTableReport(outputDir, summaryTableFilename, summaryTable, tableInfo, false, bCsvTable, bSpreadsheetFriendlyStrings, null, telemetrySessionLabel, telemetryTestIdentity);
 				}
 				if (bCollateTable)
 				{
-					WriteSummaryTableReport(outputDir, collatedTableFilename, summaryTable, tableInfo, true, bCsvTable, bSpreadsheetFriendlyStrings, weightByColumnName, telemetrySessionLabel, telemetryDataSource);
+					WriteSummaryTableReport(outputDir, collatedTableFilename, summaryTable, tableInfo, true, bCsvTable, bSpreadsheetFriendlyStrings, weightByColumnName, telemetrySessionLabel, telemetryTestIdentity);
 				}
 
 				// EmailTable is hardcoded to use the condensed type
@@ -892,7 +892,7 @@ namespace PerfReportTool
 			WriteSummaryTableReport(outputDir, filenameWithoutExtension, table, tableInfo, bCollated, bToCSV, bSpreadsheetFriendlyStrings, weightByColumnName);
 		}
 
-		void WriteSummaryTableReport(string outputDir, string filenameWithoutExtension, SummaryTable table, SummaryTableInfo tableInfo, bool bCollated, bool bToCSV, bool bSpreadsheetFriendlyStrings, string weightByColumnNameOverride, string telemetrySessionLabel = null, string telemetryDataSource = null)
+		void WriteSummaryTableReport(string outputDir, string filenameWithoutExtension, SummaryTable table, SummaryTableInfo tableInfo, bool bCollated, bool bToCSV, bool bSpreadsheetFriendlyStrings, string weightByColumnNameOverride, string telemetrySessionLabel = null, string telemetryTestIdentity = null)
 		{
 			string weightByColumnName = weightByColumnNameOverride != null ? weightByColumnNameOverride : tableInfo.weightByColumn;
 			if (GetBoolArg("noWeightedAvg"))
@@ -1003,7 +1003,7 @@ namespace PerfReportTool
 				}
 
 				string summaryTableName = GetArg("summaryTable", null);
-				Telemetry.EnqueuePostSummaryTableMetricEvent(summaryTableName, reconstitutedTable, bCollated, telemetrySessionLabel, telemetryDataSource);
+				Telemetry.EnqueuePostSummaryTableMetricEvent(summaryTableName, reconstitutedTable, bCollated, telemetrySessionLabel, telemetryTestIdentity);
 			}
 		}
 
@@ -1224,7 +1224,7 @@ namespace PerfReportTool
 					graphGenerator = new GraphGenerator(csvFile.GetFinalCsv(), csvFile.filename);
 				}
 
-				// Generate all the graphs asyncronously
+				// Generate all the graphs asynchronously
 				foreach (ReportGraph graph in reportTypeInfo.graphs)
 				{
 					if (graph.settings.statString.isSet == false)
diff --git a/Engine/Source/Programs/CSVTools/PerfReportTool/Telemetry.cs b/Engine/Source/Programs/CSVTools/PerfReportTool/Telemetry.cs
--- a/Engine/Source/Programs/CSVTools/PerfReportTool/Telemetry.cs
+++ b/Engine/Source/Programs/CSVTools/PerfReportTool/Telemetry.cs
@@ -27,7 +27,7 @@ namespace PerfReportTool
 		private const string DefaultAppEnvironment = "Dev";
 		private const string EventName = "State.Performance";
 		private const string DefaultSummaryTableName = "default";
-		private const int CurrentSchemaVersion = 1;
+		private const int CurrentSchemaVersion = 2;
 		private const int DefaultTimeoutSeconds = 15;
 		private static ServiceProvider _serviceProvider = null;
 		private static IHordeClient _hordeClient = null;
@@ -168,9 +168,9 @@ namespace PerfReportTool
 		/// </summary>
 		/// <param name="rowData">The rowData data to post.</param>
 		/// <param name="sessionLabel">Optional session label used to correlate telemetry events.</param>
-		/// <param name="dataSource">Optional data source used to group telemetry events.</param>
+		/// <param name="testIdentity">Optional identity used to group telemetry events.</param>
 		/// <returns>True if telemetry was emitted successfully, false otherwise.</returns>
-		internal async static Task<bool> PostSummaryTableMetricEventAsync(SummaryTableRowData rowData, string sessionLabel = null, string dataSource = null)
+		internal async static Task<bool> PostSummaryTableMetricEventAsync(SummaryTableRowData rowData, string sessionLabel = null, string testIdentity = null)
 		{
 			if (HordeClient == null)
 			{
@@ -182,7 +182,7 @@ namespace PerfReportTool
 
 			try
 			{
-				telemetryPayload = CreatePayload(rowData, sessionLabel, dataSource);
+				telemetryPayload = CreatePayload(rowData, sessionLabel, testIdentity);
 			}
 			catch (InvalidOperationException ex)
 			{
@@ -198,9 +198,9 @@ namespace PerfReportTool
 		/// </summary>
 		/// <param name="rowData">The rowData data to post.</param>
 		/// <param name="sessionLabel">Optional session label used to correlate telemetry events.</param>
-		/// <param name="dataSource">Optional data source used to group telemetry events.</param>
+		/// <param name="testIdentity">Optional identity used to group telemetry events.</param>
 		/// <returns>A readonly list containing the queued tasks.</returns>
-		internal static IReadOnlyList<Task<bool>> EnqueuePostSummaryTableMetricEvent(SummaryTableRowData rowData, string sessionLabel = null, string dataSource = null)
+		internal static IReadOnlyList<Task<bool>> EnqueuePostSummaryTableMetricEvent(SummaryTableRowData rowData, string sessionLabel = null, string testIdentity = null)
 		{
 			if (HordeClient == null)
 			{
@@ -210,7 +210,7 @@ namespace PerfReportTool
 
 			lock (_lock)
 			{
-				_queuedRequests.Add(PostSummaryTableMetricEventAsync(rowData, sessionLabel, dataSource));
+				_queuedRequests.Add(PostSummaryTableMetricEventAsync(rowData, sessionLabel, testIdentity));
 			}
 
 			return _queuedRequests;
@@ -221,9 +221,9 @@ namespace PerfReportTool
 		/// </summary>
 		/// <param name="rowData">The rowData data to post.</param>
 		/// <param name="sessionLabel">Optional session label used to correlate telemetry events.</param>
-		/// <param name="dataSource">Optional data source used to group telemetry events.</param>
+		/// <param name="testIdentity">Optional identity used to group telemetry events.</param>
 		/// <returns>True if telemetry was emitted successfully, false otherwise.</returns>
-		internal static bool PostSummaryTableMetricEvent(SummaryTableRowData rowData, string sessionLabel = null, string dataSource = null)
+		internal static bool PostSummaryTableMetricEvent(SummaryTableRowData rowData, string sessionLabel = null, string testIdentity = null)
 		{
 			if (HordeClient == null)
 			{
@@ -235,7 +235,7 @@ namespace PerfReportTool
 
 			try
 			{
-				telemetryPayload = CreatePayload(rowData, sessionLabel, dataSource);
+				telemetryPayload = CreatePayload(rowData, sessionLabel, testIdentity);
 			}
 			catch (InvalidOperationException ex)
 			{
@@ -246,7 +246,7 @@ namespace PerfReportTool
 			return EmitPayloadsAsync([telemetryPayload]).GetAwaiter().GetResult();
 		}
 
-		private static Hashtable CreatePayload(SummaryTableRowData rowData, string sessionLabel = null, string dataSource = null)
+		private static Hashtable CreatePayload(SummaryTableRowData rowData, string sessionLabel = null, string testIdentity = null)
 		{
 			Dictionary<string, dynamic> jsonDict = rowData.ToJsonDict(false, false);
 			string metadataMetricKey = SummaryTableElement.Type.CsvMetadata.ToString();
@@ -290,9 +290,9 @@ namespace PerfReportTool
 
 			telemetryPayload = StandardTelemetryMetadata.SanitizeHashtable(telemetryPayload);
 
-			if (!telemetryPayload.ContainsKey("DataSource") && !String.IsNullOrEmpty(dataSource))
+			if (!telemetryPayload.ContainsKey("TestIdentity") && !String.IsNullOrEmpty(testIdentity))
 			{
-				telemetryPayload["DataSource"] = dataSource;
+				telemetryPayload["TestIdentity"] = testIdentity;
 			}
 
 			HordeContextTelemetryRecord telemetryRecord = new HordeContextTelemetryRecord($"{EventName}.Global", CurrentSchemaVersion, _sessionId, sessionLabel);
@@ -335,9 +335,9 @@ namespace PerfReportTool
 		/// <param name="tableData">The table to post as an event.</param>
 		/// <param name="isCollated">If the table represents collated data or not.</param>
 		/// <param name="sessionLabel">Optional session label used to correlate telemetry events.</param>
-		/// <param name="dataSource">Optional data source used to group telemetry events.</param>
+		/// <param name="testIdentity">Optional identity used to group telemetry events.</param>
 		/// <returns>A readonly list containing the queued tasks.</returns>
-		internal static IReadOnlyList<Task<bool>> EnqueuePostSummaryTableMetricEvent(string summaryTableName, SummaryTable tableData, bool isCollated, string sessionLabel = null, string dataSource = null)
+		internal static IReadOnlyList<Task<bool>> EnqueuePostSummaryTableMetricEvent(string summaryTableName, SummaryTable tableData, bool isCollated, string sessionLabel = null, string testIdentity = null)
 		{
 			if (String.IsNullOrEmpty(summaryTableName))
 			{
@@ -363,9 +363,9 @@ namespace PerfReportTool
 
 				telemetryPayload["SummaryName"] = summaryTableName;
 
-				if (!telemetryPayload.ContainsKey("DataSource") && !String.IsNullOrEmpty(dataSource))
+				if (!telemetryPayload.ContainsKey("TestIdentity") && !String.IsNullOrEmpty(testIdentity))
 				{
-					telemetryPayload["DataSource"] = dataSource;
+					telemetryPayload["TestIdentity"] = testIdentity;
 				}
 
 				telemetryPayload["Collated"] = isCollated;
diff --git a/Engine/Source/Programs/Shared/EpicGames.Analytics.PerformanceTrends/PerformanceTrendTelemetry.cs b/Engine/Source/Programs/Shared/EpicGames.Analytics.PerformanceTrends/PerformanceTrendTelemetry.cs
--- a/Engine/Source/Programs/Shared/EpicGames.Analytics.PerformanceTrends/PerformanceTrendTelemetry.cs
+++ b/Engine/Source/Programs/Shared/EpicGames.Analytics.PerformanceTrends/PerformanceTrendTelemetry.cs
@@ -30,6 +30,11 @@ namespace EpicGames.Analytics.PerformanceTrends
 		internal const string EndTimestamp = "endtimestamp";
 		internal const string BuildSuffix = "buildsuffix";
 		internal const string TestBuildIsPreflight = "testbuildispreflight";
+		internal const string Cpu = "cpu";
+		internal const string Gpu = "gpu";
+		internal const string VerbatimRhiName = "verbatimrhiname";
+		internal const string DeviceProfile = "deviceprofile";
+
 
 		#endregion -- Internal Constants --
 
@@ -102,7 +107,11 @@ namespace EpicGames.Analytics.PerformanceTrends
 				StartTimestamp,
 				EndTimestamp,
 				BuildSuffix,
-				TestBuildIsPreflight
+				TestBuildIsPreflight,
+				Gpu,
+				Cpu,
+				VerbatimRhiName,
+				DeviceProfile
 		};
 
 		/// <summary>
@@ -194,19 +203,19 @@ namespace EpicGames.Analytics.PerformanceTrends
 		public string? GauntletSubTest { get; init; }
 
 		/// <summary>
-		/// Thether this represents collated data, or not.
+		/// Whether this represents collated data, or not.
 		/// </summary>
 		[Column("collated")]
 		public bool Collated { get; init; }
 
 		/// <summary>
-		/// The branch the underyling test build was produced from.
+		/// The branch the underlying test build was produced from.
 		/// </summary>
 		[Column("branch")]
 		public string? Branch { get; init; }
 
 		/// <summary>
-		/// The changelist the underyling test build was produced from.
+		/// The changelist the underlying test build was produced from.
 		/// </summary>
 		[Column("changelist")]
 		public int? Changelist { get; init; }
@@ -238,11 +247,35 @@ namespace EpicGames.Analytics.PerformanceTrends
 		[Column("testbuildispreflight")]
 		public int TestBuildIsPreflight { get; init; }
 
+		/// <summary>
+		/// The CPU of the hardware the performance telemetry was generated from.
+		/// </summary>
+		[Column("cpu")]
+		public string? Cpu { get; init; }
+
+		/// <summary>
+		/// The GPU of the hardware the performance telemetry was generated from.
+		/// </summary>
+		[Column("gpu")]
+		public string? Gpu { get; init; }
+
+		/// <summary>
+		/// The RHI name of the hardware the performance telemetry was generated from.
+		/// </summary>
+		[Column("verbatimrhiname")]
+		public string? VerbatimRhiName { get; init; }
+
+		/// <summary>
+		/// The device profile the performance telemetry was generated from.
+		/// </summary>
+		[Column("deviceprofile")]
+		public string? DeviceProfile { get; init; }
+
 		/// <summary>
 		/// The data source of the performance trend telemetry.
 		/// </summary>
-		[Column("data_source")]
-		public string? DataSource { get; init; }
+		[Column("test_identity")]
+		public string? TestIdentity { get; init; }
 
 		/// <summary>
 		/// Summary table name.
@@ -279,15 +312,15 @@ namespace EpicGames.Analytics.PerformanceTrends
 		/// <param name="commitId">The commitId.</param>
 		/// <param name="summaryName">The summary name.</param>
 		/// <param name="collated">Whether this record represents pre-aggregation.</param>
-		/// <param name="dataSource">The context in which the performance trend was generated.</param>
+		/// <param name="testIdentity">The context in which the performance trend was generated.</param>
 #pragma warning disable CA1054 // URI-like parameters should not be strings
-		protected PerformanceTrendTelemetry(Hashtable source, string eventName, int schemaVersion, string? sessionId, string? sessionLabel, bool isBuildMachine, string? hordeUrlStr, string? streamId, string? templateId, string? jobId, string? stepId, int? commitIdOrdered, string? commitId, string summaryName, bool collated, string? dataSource)
+		protected PerformanceTrendTelemetry(Hashtable source, string eventName, int schemaVersion, string? sessionId, string? sessionLabel, bool isBuildMachine, string? hordeUrlStr, string? streamId, string? templateId, string? jobId, string? stepId, int? commitIdOrdered, string? commitId, string summaryName, bool collated, string? testIdentity)
 			: base(eventName, schemaVersion, sessionId, sessionLabel, isBuildMachine, hordeUrlStr, streamId, templateId, jobId, stepId, commitIdOrdered, commitId)
 #pragma warning restore CA1054 // URI-like parameters should not be strings
 		{
 			SummaryName = summaryName;
 			Collated = collated;
-			DataSource = dataSource;
+			TestIdentity = testIdentity;
 
 			EngineVersion = GetString(source, StandardTelemetryMetadata.EngineVersion);
 			EngineReleaseVersion = GetString(source, StandardTelemetryMetadata.EngineReleaseVersion);
@@ -304,6 +337,10 @@ namespace EpicGames.Analytics.PerformanceTrends
 			EndTimestamp = GetLong(source, StandardTelemetryMetadata.EndTimestamp);
 			BuildSuffix = GetString(source, StandardTelemetryMetadata.BuildSuffix);
 			TestBuildIsPreflight = GetInt(source, StandardTelemetryMetadata.TestBuildIsPreflight) ?? 0;
+			Cpu = GetString(source, StandardTelemetryMetadata.Cpu);
+			Gpu = GetString(source, StandardTelemetryMetadata.Gpu);
+			VerbatimRhiName = GetString(source, StandardTelemetryMetadata.VerbatimRhiName);
+			DeviceProfile = GetString(source, StandardTelemetryMetadata.DeviceProfile);
 		}
 
 		#region -- Private API --
