CsvTools: Don't build directly to binaries path and instead use a publish script.
Use Publish.bat to publish to the output directory.
PublishWithP4 can be used to check out the files in p4 after publishing

diff --git a/Engine/Source/Programs/CSVTools/CSVCollate/CSVCollate.csproj b/Engine/Source/Programs/CSVTools/CSVCollate/CSVCollate.csproj
--- a/Engine/Source/Programs/CSVTools/CSVCollate/CSVCollate.csproj
+++ b/Engine/Source/Programs/CSVTools/CSVCollate/CSVCollate.csproj
@@ -1,29 +1,13 @@
 <Project Sdk="Microsoft.NET.Sdk">
   <PropertyGroup>
-    <TargetFramework>net8.0</TargetFramework>
     <OutputType>Exe</OutputType>
-    <GenerateAssemblyInfo>false</GenerateAssemblyInfo>
-    <AppendTargetFrameworkToOutputPath>false</AppendTargetFrameworkToOutputPath>
     <AssemblyTitle>GraphicsPerfreportTool</AssemblyTitle>
     <Product>GraphicsPerfreportTool</Product>
-    <Copyright>Copyright %28C%29 Microsoft. All rights reserved.
-Copyright Epic Games, Inc. All Rights Reserved.</Copyright>
-    <AssemblyVersion>1.0.0.0</AssemblyVersion>
-    <FileVersion>1.0.0.0</FileVersion>
-    <RuntimeIdentifiers>win-x64;osx-x64;win-arm64;osx-arm64;linux-x64;linux-arm64</RuntimeIdentifiers>
-  </PropertyGroup>
-  <PropertyGroup Condition=" '$(Configuration)|$(Platform)' == 'Debug|AnyCPU' ">
-    <PlatformTarget>x64</PlatformTarget>
-  </PropertyGroup>
-  <PropertyGroup Condition=" '$(Configuration)|$(Platform)' == 'Release|AnyCPU' ">
-    <PlatformTarget>anycpu</PlatformTarget>
-     <OutputPath Condition="'$(OutputPath)' == ''">..\..\..\..\Binaries\DotNET\CsvTools\</OutputPath>
   </PropertyGroup>
   <ItemGroup>
     <ProjectReference Include="..\CsvStats\CsvStats.csproj" />
   </ItemGroup>
   <ItemGroup>
-    <PackageReference Include="Microsoft.CSharp" Version="4.7.0" />
     <PackageReference Include="System.Data.DataSetExtensions" Version="4.5.0" />
   </ItemGroup>
 </Project>
\ No newline at end of file
diff --git a/Engine/Source/Programs/CSVTools/CSVCollate/Properties/AssemblyInfo.cs b/Engine/Source/Programs/CSVTools/CSVCollate/Properties/AssemblyInfo.cs
--- a/Engine/Source/Programs/CSVTools/CSVCollate/Properties/AssemblyInfo.cs
+++ b/Engine/Source/Programs/CSVTools/CSVCollate/Properties/AssemblyInfo.cs
@@ -1,16 +1,7 @@
 // Copyright (C) Microsoft. All rights reserved.
 // Copyright Epic Games, Inc. All Rights Reserved.
 
-using System.Reflection;
-using System.Runtime.CompilerServices;
 using System.Runtime.InteropServices;
-[assembly: AssemblyTrademark("")]
-[assembly: AssemblyCulture("")]
-
-// Setting ComVisible to false makes the types in this assembly not visible 
-// to COM components.  If you need to access a type in this assembly from 
-// COM, set the ComVisible attribute to true on that type.
-[assembly: ComVisible(false)]
 
 // The following GUID is for the ID of the typelib if this project is exposed to COM
 [assembly: Guid("ae427314-972a-4178-8c36-fc3efcd164e3")]
diff --git a/Engine/Source/Programs/CSVTools/CSVFilter/CSVFilter.csproj b/Engine/Source/Programs/CSVTools/CSVFilter/CSVFilter.csproj
--- a/Engine/Source/Programs/CSVTools/CSVFilter/CSVFilter.csproj
+++ b/Engine/Source/Programs/CSVTools/CSVFilter/CSVFilter.csproj
@@ -1,29 +1,13 @@
 <Project Sdk="Microsoft.NET.Sdk">
   <PropertyGroup>
-    <TargetFramework>net8.0</TargetFramework>
     <OutputType>Exe</OutputType>
-    <GenerateAssemblyInfo>false</GenerateAssemblyInfo>
-    <AppendTargetFrameworkToOutputPath>false</AppendTargetFrameworkToOutputPath>
     <AssemblyTitle>CSVFilter</AssemblyTitle>
     <Product>CSVFilter</Product>
-    <Copyright>Copyright %28C%29 Microsoft. All rights reserved.
-Copyright Epic Games, Inc. All Rights Reserved.</Copyright>
-    <AssemblyVersion>1.0.0.0</AssemblyVersion>
-    <FileVersion>1.0.0.0</FileVersion>
-    <RuntimeIdentifiers>win-x64;osx-x64;win-arm64;osx-arm64;linux-x64;linux-arm64</RuntimeIdentifiers>
-  </PropertyGroup>
-  <PropertyGroup Condition=" '$(Configuration)|$(Platform)' == 'Debug|AnyCPU' ">
-    <PlatformTarget>x64</PlatformTarget>
-  </PropertyGroup>
-  <PropertyGroup Condition=" '$(Configuration)|$(Platform)' == 'Release|AnyCPU' ">
-    <PlatformTarget>anycpu</PlatformTarget>
-     <OutputPath Condition="'$(OutputPath)' == ''">..\..\..\..\Binaries\DotNET\CsvTools\</OutputPath>
   </PropertyGroup>
   <ItemGroup>
     <ProjectReference Include="..\CsvStats\CsvStats.csproj" />
   </ItemGroup>
   <ItemGroup>
-    <PackageReference Include="Microsoft.CSharp" Version="4.7.0" />
     <PackageReference Include="System.Data.DataSetExtensions" Version="4.5.0" />
   </ItemGroup>
 </Project>
\ No newline at end of file
diff --git a/Engine/Source/Programs/CSVTools/CSVFilter/Properties/AssemblyInfo.cs b/Engine/Source/Programs/CSVTools/CSVFilter/Properties/AssemblyInfo.cs
--- a/Engine/Source/Programs/CSVTools/CSVFilter/Properties/AssemblyInfo.cs
+++ b/Engine/Source/Programs/CSVTools/CSVFilter/Properties/AssemblyInfo.cs
@@ -1,16 +1,7 @@
 // Copyright (C) Microsoft. All rights reserved.
 // Copyright Epic Games, Inc. All Rights Reserved.
 
-using System.Reflection;
-using System.Runtime.CompilerServices;
 using System.Runtime.InteropServices;
-[assembly: AssemblyTrademark("")]
-[assembly: AssemblyCulture("")]
-
-// Setting ComVisible to false makes the types in this assembly not visible 
-// to COM components.  If you need to access a type in this assembly from 
-// COM, set the ComVisible attribute to true on that type.
-[assembly: ComVisible(false)]
 
 // The following GUID is for the ID of the typelib if this project is exposed to COM
 [assembly: Guid("352b7c21-e47c-44d0-969a-bc0e68530df6")]
diff --git a/Engine/Source/Programs/CSVTools/CSVSplit/CSVSplit.csproj b/Engine/Source/Programs/CSVTools/CSVSplit/CSVSplit.csproj
--- a/Engine/Source/Programs/CSVTools/CSVSplit/CSVSplit.csproj
+++ b/Engine/Source/Programs/CSVTools/CSVSplit/CSVSplit.csproj
@@ -1,29 +1,13 @@
 <Project Sdk="Microsoft.NET.Sdk">
   <PropertyGroup>
-    <TargetFramework>net8.0</TargetFramework>
     <OutputType>Exe</OutputType>
-    <AppendTargetFrameworkToOutputPath>false</AppendTargetFrameworkToOutputPath>
-    <GenerateAssemblyInfo>false</GenerateAssemblyInfo>
     <AssemblyTitle>CSVSplit</AssemblyTitle>
     <Product>CSVSplit</Product>
-    <Copyright>Copyright %28C%29 Microsoft. All rights reserved.
-Copyright Epic Games, Inc. All Rights Reserved.</Copyright>
-    <AssemblyVersion>1.0.0.0</AssemblyVersion>
-    <FileVersion>1.0.0.0</FileVersion>
-    <RuntimeIdentifiers>win-x64;osx-x64;win-arm64;osx-arm64;linux-x64;linux-arm64</RuntimeIdentifiers>
-  </PropertyGroup>
-  <PropertyGroup Condition=" '$(Configuration)|$(Platform)' == 'Debug|AnyCPU' ">
-    <PlatformTarget>x64</PlatformTarget>
-  </PropertyGroup>
-  <PropertyGroup Condition=" '$(Configuration)|$(Platform)' == 'Release|AnyCPU' ">
-    <PlatformTarget>anycpu</PlatformTarget>
-     <OutputPath Condition="'$(OutputPath)' == ''">..\..\..\..\Binaries\DotNET\CsvTools\</OutputPath>
   </PropertyGroup>
   <ItemGroup>
     <ProjectReference Include="..\CsvStats\CsvStats.csproj" />
   </ItemGroup>
   <ItemGroup>
-    <PackageReference Include="Microsoft.CSharp" Version="4.7.0" />
     <PackageReference Include="System.Data.DataSetExtensions" Version="4.5.0" />
   </ItemGroup>
 </Project>
\ No newline at end of file
diff --git a/Engine/Source/Programs/CSVTools/CSVSplit/Properties/AssemblyInfo.cs b/Engine/Source/Programs/CSVTools/CSVSplit/Properties/AssemblyInfo.cs
--- a/Engine/Source/Programs/CSVTools/CSVSplit/Properties/AssemblyInfo.cs
+++ b/Engine/Source/Programs/CSVTools/CSVSplit/Properties/AssemblyInfo.cs
@@ -1,16 +1,7 @@
 // Copyright (C) Microsoft. All rights reserved.
 // Copyright Epic Games, Inc. All Rights Reserved.
 
-using System.Reflection;
-using System.Runtime.CompilerServices;
 using System.Runtime.InteropServices;
-[assembly: AssemblyTrademark("")]
-[assembly: AssemblyCulture("")]
-
-// Setting ComVisible to false makes the types in this assembly not visible 
-// to COM components.  If you need to access a type in this assembly from 
-// COM, set the ComVisible attribute to true on that type.
-[assembly: ComVisible(false)]
 
 // The following GUID is for the ID of the typelib if this project is exposed to COM
 [assembly: Guid("a7c8e49c-6a8a-40f4-9dfd-8135acf7bf98")]
diff --git a/Engine/Source/Programs/CSVTools/CSVToSVG/CSVToSVG.csproj b/Engine/Source/Programs/CSVTools/CSVToSVG/CSVToSVG.csproj
--- a/Engine/Source/Programs/CSVTools/CSVToSVG/CSVToSVG.csproj
+++ b/Engine/Source/Programs/CSVTools/CSVToSVG/CSVToSVG.csproj
@@ -1,34 +1,19 @@
 <Project Sdk="Microsoft.NET.Sdk">
   <PropertyGroup>
-    <TargetFramework>net8.0</TargetFramework>
     <OutputType>Exe</OutputType>
-    <GenerateAssemblyInfo>false</GenerateAssemblyInfo>
-    <AppendTargetFrameworkToOutputPath>false</AppendTargetFrameworkToOutputPath>
     <RuntimeIdentifiers>win-x64;osx-x64;win-arm64;osx-arm64;linux-x64;linux-arm64</RuntimeIdentifiers>
   </PropertyGroup>
-  <PropertyGroup Condition=" '$(Configuration)|$(Platform)' == 'Debug|AnyCPU' ">
-    <PlatformTarget>x64</PlatformTarget>
-  </PropertyGroup>
-  <PropertyGroup Condition=" '$(Configuration)|$(Platform)' == 'Release|AnyCPU' ">
-    <PlatformTarget>anycpu</PlatformTarget>
-     <OutputPath Condition="'$(OutputPath)' == ''">..\..\..\..\Binaries\DotNET\CsvTools\</OutputPath>
-  </PropertyGroup>
   <PropertyGroup>
     <StartupObject>
     </StartupObject>
     <AssemblyTitle>CSVToSVG</AssemblyTitle>
     <Product>CSVToSVG</Product>
-    <Copyright>Copyright %28C%29 Microsoft. All rights reserved.
-Copyright Epic Games, Inc. All Rights Reserved.</Copyright>
-    <AssemblyVersion>1.0.0.0</AssemblyVersion>
-    <FileVersion>1.0.0.0</FileVersion>
   </PropertyGroup>
   <ItemGroup>
     <ProjectReference Include="..\CsvStats\CsvStats.csproj" />
     <ProjectReference Include="..\CsvToSvgLib\CsvToSvgLib.csproj" />
   </ItemGroup>
   <ItemGroup>
-    <PackageReference Include="Microsoft.CSharp" Version="4.7.0" />
     <PackageReference Include="System.Data.DataSetExtensions" Version="4.5.0" />
   </ItemGroup>
 </Project>
\ No newline at end of file
diff --git a/Engine/Source/Programs/CSVTools/CSVToSVG/Properties/AssemblyInfo.cs b/Engine/Source/Programs/CSVTools/CSVToSVG/Properties/AssemblyInfo.cs
--- a/Engine/Source/Programs/CSVTools/CSVToSVG/Properties/AssemblyInfo.cs
+++ b/Engine/Source/Programs/CSVTools/CSVToSVG/Properties/AssemblyInfo.cs
@@ -1,16 +1,7 @@
 // Copyright (C) Microsoft. All rights reserved.
 // Copyright Epic Games, Inc. All Rights Reserved.
 
-using System.Reflection;
-using System.Runtime.CompilerServices;
 using System.Runtime.InteropServices;
-[assembly: AssemblyTrademark("")]
-[assembly: AssemblyCulture("")]
-
-// Setting ComVisible to false makes the types in this assembly not visible 
-// to COM components.  If you need to access a type in this assembly from 
-// COM, set the ComVisible attribute to true on that type.
-[assembly: ComVisible(false)]
 
 // The following GUID is for the ID of the typelib if this project is exposed to COM
 [assembly: Guid("dc6aed41-e787-4807-a139-52b9557567d4")]
diff --git a/Engine/Source/Programs/CSVTools/CsvConvert/CsvConvert.csproj b/Engine/Source/Programs/CSVTools/CsvConvert/CsvConvert.csproj
--- a/Engine/Source/Programs/CSVTools/CsvConvert/CsvConvert.csproj
+++ b/Engine/Source/Programs/CSVTools/CsvConvert/CsvConvert.csproj
@@ -1,28 +1,13 @@
 ﻿<Project Sdk="Microsoft.NET.Sdk">
   <PropertyGroup>
-    <TargetFramework>net8.0</TargetFramework>
     <OutputType>Exe</OutputType>
-    <GenerateAssemblyInfo>false</GenerateAssemblyInfo>
-    <AppendTargetFrameworkToOutputPath>false</AppendTargetFrameworkToOutputPath>
     <AssemblyTitle>CsvConvert</AssemblyTitle>
     <Product>CsvConvert</Product>
-    <Copyright>Copyright Epic Games, Inc. All rights reserved.</Copyright>
-    <AssemblyVersion>1.0.0.0</AssemblyVersion>
-    <FileVersion>1.0.0.0</FileVersion>
-    <RuntimeIdentifiers>win-x64;osx-x64;win-arm64;osx-arm64;linux-x64;linux-arm64</RuntimeIdentifiers>
-  </PropertyGroup>
-  <PropertyGroup Condition=" '$(Configuration)|$(Platform)' == 'Debug|AnyCPU' ">
-    <PlatformTarget>x64</PlatformTarget>
-  </PropertyGroup>
-  <PropertyGroup Condition=" '$(Configuration)|$(Platform)' == 'Release|AnyCPU' ">
-    <PlatformTarget>anycpu</PlatformTarget>
-     <OutputPath Condition="'$(OutputPath)' == ''">..\..\..\..\Binaries\DotNET\CsvTools\</OutputPath>
   </PropertyGroup>
   <ItemGroup>
     <ProjectReference Include="..\CsvStats\CsvStats.csproj" />
   </ItemGroup>
   <ItemGroup>
-    <PackageReference Include="Microsoft.CSharp" Version="4.7.0" />
     <PackageReference Include="System.Data.DataSetExtensions" Version="4.5.0" />
   </ItemGroup>
 </Project>
\ No newline at end of file
diff --git a/Engine/Source/Programs/CSVTools/CsvConvert/Properties/AssemblyInfo.cs b/Engine/Source/Programs/CSVTools/CsvConvert/Properties/AssemblyInfo.cs
--- a/Engine/Source/Programs/CSVTools/CsvConvert/Properties/AssemblyInfo.cs
+++ b/Engine/Source/Programs/CSVTools/CsvConvert/Properties/AssemblyInfo.cs
@@ -1,15 +1,6 @@
 ﻿// Copyright Epic Games, Inc. All Rights Reserved.
 
-using System.Reflection;
-using System.Runtime.CompilerServices;
 using System.Runtime.InteropServices;
-[assembly: AssemblyTrademark("")]
-[assembly: AssemblyCulture("")]
-
-// Setting ComVisible to false makes the types in this assembly not visible 
-// to COM components.  If you need to access a type in this assembly from 
-// COM, set the ComVisible attribute to true on that type.
-[assembly: ComVisible(false)]
 
 // The following GUID is for the ID of the typelib if this project is exposed to COM
 [assembly: Guid("5110c345-5449-40e2-9728-6197b23b5b37")]
diff --git a/Engine/Source/Programs/CSVTools/CsvInfo/CsvInfo.csproj b/Engine/Source/Programs/CSVTools/CsvInfo/CsvInfo.csproj
--- a/Engine/Source/Programs/CSVTools/CsvInfo/CsvInfo.csproj
+++ b/Engine/Source/Programs/CSVTools/CsvInfo/CsvInfo.csproj
@@ -1,28 +1,13 @@
 ﻿<Project Sdk="Microsoft.NET.Sdk">
   <PropertyGroup>
-    <TargetFramework>net8.0</TargetFramework>
     <OutputType>Exe</OutputType>
-    <GenerateAssemblyInfo>false</GenerateAssemblyInfo>
-    <AppendTargetFrameworkToOutputPath>false</AppendTargetFrameworkToOutputPath>
     <AssemblyTitle>CsvInfo</AssemblyTitle>
     <Product>CsvInfo</Product>
-    <Copyright>Copyright Epic Games, Inc. All rights reserved.</Copyright>
-    <AssemblyVersion>1.0.0.0</AssemblyVersion>
-    <FileVersion>1.0.0.0</FileVersion>
-    <RuntimeIdentifiers>win-x64;osx-x64;win-arm64;osx-arm64;linux-x64;linux-arm64</RuntimeIdentifiers>
-  </PropertyGroup>
-  <PropertyGroup Condition=" '$(Configuration)|$(Platform)' == 'Debug|AnyCPU' ">
-    <PlatformTarget>x64</PlatformTarget>
-  </PropertyGroup>
-  <PropertyGroup Condition=" '$(Configuration)|$(Platform)' == 'Release|AnyCPU' ">
-    <PlatformTarget>anycpu</PlatformTarget>
-     <OutputPath Condition="'$(OutputPath)' == ''">..\..\..\..\Binaries\DotNET\CsvTools\</OutputPath>
   </PropertyGroup>
   <ItemGroup>
     <ProjectReference Include="..\CsvStats\CsvStats.csproj" />
   </ItemGroup>
   <ItemGroup>
-    <PackageReference Include="Microsoft.CSharp" Version="4.7.0" />
     <PackageReference Include="System.Data.DataSetExtensions" Version="4.5.0" />
   </ItemGroup>
 </Project>
\ No newline at end of file
diff --git a/Engine/Source/Programs/CSVTools/CsvInfo/Properties/AssemblyInfo.cs b/Engine/Source/Programs/CSVTools/CsvInfo/Properties/AssemblyInfo.cs
--- a/Engine/Source/Programs/CSVTools/CsvInfo/Properties/AssemblyInfo.cs
+++ b/Engine/Source/Programs/CSVTools/CsvInfo/Properties/AssemblyInfo.cs
@@ -1,15 +1,6 @@
 ﻿// Copyright Epic Games, Inc. All Rights Reserved.
 
-using System.Reflection;
-using System.Runtime.CompilerServices;
 using System.Runtime.InteropServices;
-[assembly: AssemblyTrademark("")]
-[assembly: AssemblyCulture("")]
-
-// Setting ComVisible to false makes the types in this assembly not visible 
-// to COM components.  If you need to access a type in this assembly from 
-// COM, set the ComVisible attribute to true on that type.
-[assembly: ComVisible(false)]
 
 // The following GUID is for the ID of the typelib if this project is exposed to COM
 [assembly: Guid("5110c345-5449-40e2-9728-6197b23b5b37")]
diff --git a/Engine/Source/Programs/CSVTools/CsvStats/CsvStats.csproj b/Engine/Source/Programs/CSVTools/CsvStats/CsvStats.csproj
--- a/Engine/Source/Programs/CSVTools/CsvStats/CsvStats.csproj
+++ b/Engine/Source/Programs/CSVTools/CsvStats/CsvStats.csproj
@@ -1,24 +1,8 @@
 <Project Sdk="Microsoft.NET.Sdk">
   <PropertyGroup>
-    <AppendTargetFrameworkToOutputPath>false</AppendTargetFrameworkToOutputPath>
-    <TargetFramework>net8.0</TargetFramework>
     <OutputType>Library</OutputType>
-    <Copyright>Copyright %28C%29 Microsoft. All rights reserved.
-Copyright Epic Games, Inc. All Rights Reserved.</Copyright>
-    <AssemblyVersion>1.0.0.0</AssemblyVersion>
-    <FileVersion>1.0.0.0</FileVersion>
-    <GenerateAssemblyInfo>false</GenerateAssemblyInfo>
-    <RuntimeIdentifiers>win-x64;osx-x64;win-arm64;osx-arm64;linux-x64;linux-arm64</RuntimeIdentifiers>
-  </PropertyGroup>
-  <PropertyGroup Condition=" '$(Configuration)|$(Platform)' == 'Debug|AnyCPU' ">
-    <PlatformTarget>x64</PlatformTarget>
-  </PropertyGroup>
-  <PropertyGroup Condition=" '$(Configuration)|$(Platform)' == 'Release|AnyCPU' ">
-    <OutputPath Condition="'$(OutputPath)' == ''">..\..\..\..\Binaries\DotNET\CsvTools\</OutputPath>
-    <PlatformTarget>anycpu</PlatformTarget>
   </PropertyGroup>
   <ItemGroup>
-    <PackageReference Include="Microsoft.CSharp" Version="4.7.0" />
     <PackageReference Include="System.Data.DataSetExtensions" Version="4.5.0" />
   </ItemGroup>
 </Project>
\ No newline at end of file
diff --git a/Engine/Source/Programs/CSVTools/CsvStats/Properties/AssemblyInfo.cs b/Engine/Source/Programs/CSVTools/CsvStats/Properties/AssemblyInfo.cs
--- a/Engine/Source/Programs/CSVTools/CsvStats/Properties/AssemblyInfo.cs
+++ b/Engine/Source/Programs/CSVTools/CsvStats/Properties/AssemblyInfo.cs
@@ -1,16 +1,7 @@
 // Copyright (C) Microsoft. All rights reserved.
 // Copyright Epic Games, Inc. All Rights Reserved.
 
-using System.Reflection;
-using System.Runtime.CompilerServices;
 using System.Runtime.InteropServices;
-[assembly: AssemblyTrademark("")]
-[assembly: AssemblyCulture("")]
-
-// Setting ComVisible to false makes the types in this assembly not visible 
-// to COM components.  If you need to access a type in this assembly from 
-// COM, set the ComVisible attribute to true on that type.
-[assembly: ComVisible(false)]
 
 // The following GUID is for the ID of the typelib if this project is exposed to COM
 [assembly: Guid("54811cbd-f2b0-427f-859a-2f93318baab1")]
diff --git a/Engine/Source/Programs/CSVTools/CsvToSvgLib/CsvToSvgLib.csproj b/Engine/Source/Programs/CSVTools/CsvToSvgLib/CsvToSvgLib.csproj
--- a/Engine/Source/Programs/CSVTools/CsvToSvgLib/CsvToSvgLib.csproj
+++ b/Engine/Source/Programs/CSVTools/CsvToSvgLib/CsvToSvgLib.csproj
@@ -1,30 +1,15 @@
 ﻿<Project Sdk="Microsoft.NET.Sdk">
   <PropertyGroup>
-    <TargetFramework>net8.0</TargetFramework>
     <OutputType>Library</OutputType>
     <RootNamespace>CsvStats</RootNamespace>
-    <GenerateAssemblyInfo>false</GenerateAssemblyInfo>
-    <AppendTargetFrameworkToOutputPath>false</AppendTargetFrameworkToOutputPath>
     <AssemblyTitle>CsvToSvgLib</AssemblyTitle>
     <Product>CsvToSvgLib</Product>
-    <Copyright>Copyright %28C%29 Microsoft. All rights reserved.
-Copyright Epic Games, Inc. All Rights Reserved.</Copyright>
-    <AssemblyVersion>1.0.0.0</AssemblyVersion>
-    <FileVersion>1.0.0.0</FileVersion>
     <RuntimeIdentifiers>win-x64;osx-x64;win-arm64;osx-arm64;linux-x64;linux-arm64</RuntimeIdentifiers>
   </PropertyGroup>
-  <PropertyGroup Condition=" '$(Configuration)|$(Platform)' == 'Debug|AnyCPU' ">
-    <PlatformTarget>x64</PlatformTarget>
-  </PropertyGroup>
-  <PropertyGroup Condition=" '$(Configuration)|$(Platform)' == 'Release|AnyCPU' ">
-     <OutputPath Condition="'$(OutputPath)' == ''">..\..\..\..\Binaries\DotNET\CsvTools\</OutputPath>
-    <PlatformTarget>anycpu</PlatformTarget>
-  </PropertyGroup>
   <ItemGroup>
     <ProjectReference Include="..\CsvStats\CsvStats.csproj" />
   </ItemGroup>
   <ItemGroup>
-    <PackageReference Include="Microsoft.CSharp" Version="4.7.0" />
     <PackageReference Include="System.Data.DataSetExtensions" Version="4.5.0" />
   </ItemGroup>
 </Project>
\ No newline at end of file
diff --git a/Engine/Source/Programs/CSVTools/CsvToSvgLib/Properties/AssemblyInfo.cs b/Engine/Source/Programs/CSVTools/CsvToSvgLib/Properties/AssemblyInfo.cs
--- a/Engine/Source/Programs/CSVTools/CsvToSvgLib/Properties/AssemblyInfo.cs
+++ b/Engine/Source/Programs/CSVTools/CsvToSvgLib/Properties/AssemblyInfo.cs
@@ -1,16 +1,7 @@
 // Copyright (C) Microsoft. All rights reserved.
 // Copyright Epic Games, Inc. All Rights Reserved.
 
-using System.Reflection;
-using System.Runtime.CompilerServices;
 using System.Runtime.InteropServices;
-[assembly: AssemblyTrademark("")]
-[assembly: AssemblyCulture("")]
-
-// Setting ComVisible to false makes the types in this assembly not visible 
-// to COM components.  If you need to access a type in this assembly from 
-// COM, set the ComVisible attribute to true on that type.
-[assembly: ComVisible(false)]
 
 // The following GUID is for the ID of the typelib if this project is exposed to COM
 [assembly: Guid("54811cbd-f2b0-427f-ABCD-2f93318baab1")]
diff --git a/Engine/Source/Programs/CSVTools/PerfReportTool/LLMReportGraphs.xml b/Engine/Source/Programs/CSVTools/PerfReportTool/LLMReportGraphs.xml
new file mode 100644
--- /dev/null
+++ b/Engine/Source/Programs/CSVTools/PerfReportTool/LLMReportGraphs.xml
@@ -0,0 +1,137 @@
+<?xml version="1.0" encoding="UTF-8"?>
+
+<!-- Please use the XML schemas in "Engine/Source/Programs/CsvTools/PerfReportTool/XmlSchemas" to validate -->
+
+<!-- This structure describes how to graph out the different graphs specified in ReportTypes.xml
+       The base settings are the main settings that all of the other graphs use unless they specifically override them.
+       The main base settings are:
+         - maxY:           This is the maximum value that the graph will represent, so make sure they covers the whole dataset.
+         - stacked:        Indicates whether the output graph will be stacked.
+         - hideStatPrefix: This indicates the prefix that should be hidden in the graph legend.
+         
+       The graph node can take the same attributes as the baseSettings node (if it does it will override them, otherwise, it will use the baseSettings settings.
+       The graph has a title attribute, this is the title of the graph, this MUST match the appropriate graph node title in ReportTypes.xml.
+       The graph also has a subnode called statString, this is the string that is passed to the CSVToSVG tool, therefore * will graph all of the stats in the CSV,
+       while something like "name1 name2" will graph only the two stats with headings name1 and name2.
+       
+  -->
+<graphGroups xmlns="https://www.unrealengine.com/PerfReportTool" xmlns:xsi="http://www.w3.org/2001/XMLSchema-instance">
+	<baseSettings thickness="1" smooth="0" smoothKernelSize="-1" smoothKernelPercent="1" compression="0.1"
+				  width="1800" height="600" stacked="0" maxHierarchyDepth="2" hideStatPrefix="" showEvents="MapSetup*"/> 
+ 
+	<!-- a graph group is a convenient way to apply common settings to multiple graphs, avoiding duplication -->
+	<graphGroup name="LLMBase">
+		<baseSettings maxHierarchyDepth="4" stacked="0" hideStatPrefix="Meshes/;Audio/;" requiresDetailedStats="0" thickness="100" compression="0.05" />
+
+		<graph title="Total">
+			<statString>Total,Untracked,TrackedTotal</statString> 
+		</graph>
+
+		<!-- ignoreStats should include all parent LLM tags -->
+		<graph title="Memory Allocated" stacked="1" thickness="100" maxHierarchyDepth="1" ignoreStats="TrackedTotal;Audio;Meshes" mainStat="Total">
+			<statString>*</statString>
+		</graph>
+
+		<!-- top level LLM summary graphs -->
+		<graph title="Content Summary" stacked="1" thickness="100" maxHierarchyDepth="-1">
+			<statString>Shaders,Materials,MaterialInstance,Meshes/*,Audio/AudioSoundWaves,Animation</statString>
+		</graph>
+		<graph title="Misc Summary">
+			<statString>EngineInit,EnginePreInit,EngineMisc,TaskGraphMiscTasks,FName,GC,Localization,AssetRegistry,ConfigSystem</statString>
+		</graph>
+		<graph title="Rendering Summary">
+			<statString>RenderingThread,RenderTargets,PSO,SceneRender,RHIMisc,Graphics</statString>
+		</graph>
+		<graph title="Development Summary">
+			<statString>Stats,CsvProfiler</statString>
+		</graph>	
+
+		<!-- general graphs -->
+		<graph title="Tracked Total" maxY="8000">
+			<statString>TrackedTotal</statString>
+		</graph>
+		<graph title="Untagged" maxY="1000">
+			<statString>Untagged</statString>
+		</graph>
+		<graph title="Untracked">
+			<statString>Untracked</statString>
+		</graph>
+		<graph title="ProgramSize">
+			<statString>ProgramSize</statString>
+		</graph>
+		<graph title="EnginePreInitMemory" maxY="350">
+			<statString>EnginePreInit</statString>
+		</graph>
+		<graph title="Stats" maxY="15">
+			<statString>Stats</statString>
+		</graph>
+		<graph title="UObject" maxY="1000">
+			<statString>UObject</statString>
+		</graph>
+		<graph title="RHI Misc" maxY="250">
+			<statString>RHIMisc</statString>
+		</graph>
+		<graph title="Shaders" maxY="500">
+			<statString>Shaders</statString>
+		</graph>
+		<graph title="Render Targets" maxY="1400">
+			<statString>RenderTargets</statString>
+		</graph>
+		<graph title="Textures" maxY="1400">
+			<statString>Textures</statString>
+		</graph>
+		<graph title="AsyncLoading" maxY="300">
+			<statString>AsyncLoading</statString>
+		</graph>
+		<graph title="Materials" maxY="1000">
+			<statString>Materials</statString>
+		</graph>
+		<graph title="Audio">
+			<statString>Audio</statString>
+		</graph>
+		<graph title="Audio Detail">
+			<statString>Audio/*</statString>
+		</graph>
+		<graph title="Meshes Detail">
+			<statString>Meshes/*</statString>
+		</graph>
+		<graph title="EngineInitMemory" maxY="100">
+			<statString>EngineInit</statString>
+		</graph>
+		<graph title="Animation" maxY="50">
+			<statString>Animation</statString>
+		</graph>
+		<graph title="RenderingThreadMemory" maxY="200">
+			<statString>RenderingThread</statString>
+		</graph>
+		<graph title="LoadMapMemory" maxY="100">
+			<statString>LoadMapMisc</statString>
+		</graph>
+		<graph title="PhysX" maxY="100">
+			<statString>PhysX</statString>
+		</graph>
+		<graph title="EngineMisc" maxY="250">
+			<statString>EngineMisc</statString>
+		</graph>
+		<graph title="StreamingManager" maxY="25">
+			<statString>StreamingManager</statString>
+		</graph>
+	</graphGroup>
+
+	<graphGroup name="LLMPlatformBase">
+		<baseSettings maxHierarchyDepth="4" maxY="1024" stacked="0" hideStatPrefix="" requiresDetailedStats="0" thickness="100" compression="0.05" />
+
+		<graph title="Platform FMalloc" maxY="2500">
+			<statString>FMalloc</statString>
+		</graph>
+		<graph title="Platform GenericPlatformMallocCrash" maxY="5">
+			<statString>GenericPlatformMallocCrash</statString>
+		</graph>
+		<graph title="Platform Graphics" maxY="2000">
+			<statString>Graphics</statString>
+		</graph>
+		<graph title="Platform Thread Stack" maxY="10">
+			<statString>ThreadStack</statString>
+		</graph>
+	</graphGroup>
+</graphGroups>
diff --git a/Engine/Source/Programs/CSVTools/PerfReportTool/LLMReportTypes.xml b/Engine/Source/Programs/CSVTools/PerfReportTool/LLMReportTypes.xml
new file mode 100644
--- /dev/null
+++ b/Engine/Source/Programs/CSVTools/PerfReportTool/LLMReportTypes.xml
@@ -0,0 +1,96 @@
+<?xml version="1.0" encoding="UTF-8"?>
+
+<!-- Please use the XML schemas in "Engine/Source/Programs/CsvTools/PerfReportTool/XmlSchemas" to validate -->
+
+<!-- The report type works by having a summary type and different graphs.
+       The summary type must match one of the types in PerfReportTool, currently we only have:
+         - Peak
+         - Histogram
+       
+       To create a new type just subclass Summary and override its functions. See PeakSummary or HistogramSummary as examples.
+       
+       Then each graph has a title, budget, inSummary, inMainSummary and ignoreList.
+        - title:         Used to link the graph to the graph in ReportGraphs.xml, so they must be the same.
+        - budget:        This is the value that we are going to compare to. How much we can use.
+        - inSummary:     This indicated whether this should be included in the HTML report produced.
+        - ignoreList:    This is mainly used with wildcards, if you have stats in the CSV you want ignored for this graph, add them to this list.
+                         It's a semi-colon separated list as such: "name1;name2;..."
+       
+       NOTE: Each graph needs an entry in the appropriate graphGroup node in ReportGraphs.xml, otherwise the program crashes.
+  -->
+<root xmlns="https://www.unrealengine.com/PerfReportTool" xmlns:xsi="http://www.w3.org/2001/XMLSchema-instance">
+
+	<summaryTables>
+		<summaryTable name="default" rowSort="CSV File">
+			<filter>Total Max,Untracked Avg,CSV File,Report,frameCount,</filter>
+		</summaryTable>
+
+	</summaryTables>
+
+	<reporttypes reportGraphsFile="LLMReportGraphs.xml">
+ 
+		<!-- LLM Report -->
+		<reporttype name="LLM" title="LLM Report" vsync="0" ignoreList="">
+
+			<autodetection requiredstats="EngineInit"> 
+			</autodetection>
+ 
+			<summary type="peak"/>
+
+			<!-- top-level LLM summary graphs -->
+			<graph title="Total"                       budget="4500"    inSummary="1" />
+			<graph title="Content Summary"             budget="0"       inSummary="0"/>
+			<graph title="Rendering Summary"           budget="0"       inSummary="0"/>
+			<graph title="Misc Summary"                budget="0"       inSummary="0"/>
+			<graph title="Development Summary"         budget="0"       inSummary="0"/>
+			
+			<graph title="Tracked Total"               budget="4000"    inSummary="0"/> 
+			<!-- Sum of below -->
+			<graph title="Animation"                   budget="0"       inSummary="1"/>
+			<graph title="AsyncLoading"                budget="0"       inSummary="1"/>
+			<graph title="Audio"                       budget="0"       inSummary="1"/>
+			<graph title="Audio Detail"                budget="0"       inSummary="0"/>
+			<graph title="Meshes Detail"               budget="0"       inSummary="0"/>
+			<graph title="EngineInitMemory"            budget="0"       inSummary="1"/>
+			<graph title="EnginePreInitMemory"         budget="0"       inSummary="1"/>
+			<graph title="EngineMisc"                  budget="0"       inSummary="1"/>
+			<graph title="LoadMapMemory"               budget="0"       inSummary="1"/>
+			<graph title="Materials"                   budget="0"       inSummary="1"/>
+			<graph title="PhysX"                       budget="0"       inSummary="1"/>
+			<graph title="ProgramSize"                 budget="0"       inSummary="1"/>
+			<graph title="Render Targets"              budget="0"       inSummary="1"/>
+			<graph title="RenderingThreadMemory"       budget="0"       inSummary="1"/>
+			<graph title="RHI Misc"                    budget="0"       inSummary="1"/>
+			<graph title="Shaders"                     budget="0"       inSummary="1"/>
+			<graph title="Stats"                       budget="0"       inSummary="1"/>
+			<graph title="StreamingManager"            budget="0"       inSummary="1"/>
+			<graph title="Textures"                    budget="0"       inSummary="1"/>
+			<graph title="UObject"                     budget="0"       inSummary="1"/>
+			<graph title="Untagged"                    budget="0"       inSummary="1"/>
+			
+			<graph title="Memory Allocated"            budget="0"       inSummary="0"/>
+		</reporttype>
+
+
+		<!-- LLM Report -->
+		<reporttype name="LLMPlatformBase" title="LLM Platform Base Report" vsync="0" ignoreList="">
+			<autodetection requiredstats="LLMOverhead">
+			</autodetection>
+			
+			<summary type="peak">
+				<colourThresholds stat="Total">3000,4000,5000,6000</colourThresholds>
+			</summary>
+
+			<graph title="Total"                                budget="4500"     inSummary="1"/>
+			<graph title="Memory Allocated"                     budget="4500"     inSummary="0"/>
+			<graph title="Untagged"                             budget="0"        inSummary="1"/>
+			<graph title="Platform Graphics"                    budget="0"        inSummary="1"/>
+			<graph title="ProgramSize"                          budget="0"        inSummary="1"/>
+			<graph title="Platform FMalloc"                     budget="0"        inSummary="1"/>
+			<graph title="Platform GenericPlatformMallocCrash"  budget="0"        inSummary="1"/>
+			<graph title="Platform Thread Stack"                budget="0"        inSummary="1"/>
+
+		</reporttype>
+
+	</reporttypes>
+</root>
\ No newline at end of file
diff --git a/Engine/Source/Programs/CSVTools/PerfReportTool/PerfreportTool.csproj b/Engine/Source/Programs/CSVTools/PerfReportTool/PerfreportTool.csproj
--- a/Engine/Source/Programs/CSVTools/PerfReportTool/PerfreportTool.csproj
+++ b/Engine/Source/Programs/CSVTools/PerfReportTool/PerfreportTool.csproj
@@ -1,17 +1,8 @@
 ﻿<Project Sdk="Microsoft.NET.Sdk">
   <PropertyGroup>
-    <TargetFramework>net8.0</TargetFramework>
     <OutputType>Exe</OutputType>
-    <GenerateAssemblyInfo>false</GenerateAssemblyInfo>
-    <AppendTargetFrameworkToOutputPath>false</AppendTargetFrameworkToOutputPath>
-    <RuntimeIdentifiers>win-x64;osx-x64;win-arm64;osx-arm64;linux-x64;linux-arm64</RuntimeIdentifiers>
-  </PropertyGroup>
-  <PropertyGroup Condition=" '$(Configuration)|$(Platform)' == 'Debug|AnyCPU' ">
-    <PlatformTarget>x64</PlatformTarget>
-  </PropertyGroup>
-  <PropertyGroup Condition=" '$(Configuration)|$(Platform)' == 'Release|AnyCPU' ">
-    <PlatformTarget>anycpu</PlatformTarget>
-     <OutputPath Condition="'$(OutputPath)' == ''">..\..\..\..\Binaries\DotNET\CsvTools\</OutputPath>
+    <AssemblyTitle>PerfReportTool</AssemblyTitle>
+    <Product>PerfReportTool</Product>
   </PropertyGroup>
   <ItemGroup>
     <None Remove="Properties\AssemblyInfo.cs~RF273305d.TMP" />
@@ -23,16 +14,21 @@
     <ProjectReference Include="..\CsvToSvgLib\CsvToSvgLib.csproj" />
   </ItemGroup>
   <ItemGroup>
-    <PackageReference Include="System.Drawing.Common" Version="7.0.0" />
-    <PackageReference Include="Microsoft.CSharp" Version="4.7.0" />
+    <PackageReference Include="System.Drawing.Common" Version="8.0.24" />
     <PackageReference Include="System.Data.DataSetExtensions" Version="4.5.0" />
   </ItemGroup>
-  <PropertyGroup>
-    <AssemblyTitle>PerfReportTool</AssemblyTitle>
-    <Product>PerfReportTool</Product>
-    <Copyright>Copyright %28C%29 Microsoft. All rights reserved.
-Copyright Epic Games, Inc. All Rights Reserved.</Copyright>
-    <AssemblyVersion>1.0.0.0</AssemblyVersion>
-    <FileVersion>1.0.0.0</FileVersion>
-  </PropertyGroup>
+  <ItemGroup>
+    <None Update="LLMReportGraphs.xml">
+      <CopyToOutputDirectory>Always</CopyToOutputDirectory>
+    </None>
+    <None Update="LLMReportTypes.xml">
+      <CopyToOutputDirectory>Always</CopyToOutputDirectory>
+    </None>
+    <None Update="ReportGraphs.xml">
+      <CopyToOutputDirectory>Always</CopyToOutputDirectory>
+    </None>
+    <None Update="ReportTypes.xml">
+      <CopyToOutputDirectory>Always</CopyToOutputDirectory>
+    </None>
+  </ItemGroup>
 </Project>
\ No newline at end of file
diff --git a/Engine/Source/Programs/CSVTools/PerfReportTool/Properties/AssemblyInfo.cs b/Engine/Source/Programs/CSVTools/PerfReportTool/Properties/AssemblyInfo.cs
--- a/Engine/Source/Programs/CSVTools/PerfReportTool/Properties/AssemblyInfo.cs
+++ b/Engine/Source/Programs/CSVTools/PerfReportTool/Properties/AssemblyInfo.cs
@@ -1,16 +1,7 @@
 // Copyright (C) Microsoft. All rights reserved.
 // Copyright Epic Games, Inc. All Rights Reserved.
 
-using System.Reflection;
-using System.Runtime.CompilerServices;
 using System.Runtime.InteropServices;
-[assembly: AssemblyTrademark("")]
-[assembly: AssemblyCulture("")]
-
-// Setting ComVisible to false makes the types in this assembly not visible 
-// to COM components.  If you need to access a type in this assembly from 
-// COM, set the ComVisible attribute to true on that type.
-[assembly: ComVisible(false)]
 
 // The following GUID is for the ID of the typelib if this project is exposed to COM
 [assembly: Guid("ae427314-972a-4178-8c36-fc3efcd164e3")]
diff --git a/Engine/Source/Programs/CSVTools/PerfReportTool/ReportGraphs.xml b/Engine/Source/Programs/CSVTools/PerfReportTool/ReportGraphs.xml
new file mode 100644
--- /dev/null
+++ b/Engine/Source/Programs/CSVTools/PerfReportTool/ReportGraphs.xml
@@ -0,0 +1,178 @@
+<?xml version="1.0" encoding="UTF-8"?>
+
+<!-- Please use the XML schemas in "Engine/Source/Programs/CsvTools/PerfReportTool/XmlSchemas" to validate -->
+<graphGroups xmlns="https://www.unrealengine.com/PerfReportTool" xmlns:xsi="http://www.w3.org/2001/XMLSchema-instance">
+	<baseSettings thickness="2" smooth="1" smoothKernelSize="60" compression="0.5"
+				  width="2000" height="600" stacked="0" maxHierarchyDepth="2" hideStatPrefix="" showEvents="*" maxY="40" showAverages="1"/>
+
+	<graphGroup name="Core">
+		<baseSettings smoothKernelSize="60" smoothKernelPercent="-1"/>
+
+		<graph title="Stat Unit" showAverages="1" thickness="3">
+			<statString>frametime,gamethreadtime,renderthreadtime,gputime,rhithreadtime</statString>
+		</graph>
+
+		<graph title="Stat Unit Raw" showAverages="1" thickness="1" smooth="0" maxY="0">
+			<statString>frametime,gamethreadtime,renderthreadtime,gputime,rhithreadtime</statString>
+		</graph>
+
+		<graph title="Stat Unit 30hz" showAverages="1" maxY="60" thickness="3">
+			<statString>frametime,gamethreadtime,renderthreadtime,gputime,rhithreadtime</statString>
+		</graph>
+
+		<graph title="Stat Unit 20hz" showAverages="1" maxY="80" thickness="3">
+			<statString>frametime,gamethreadtime,renderthreadtime,gputime,rhithreadtime</statString>
+		</graph>
+
+		<graph title="Dynamic Resolution" maxY="100" thickness="4" showAverages="1">
+			<statString>DynamicResolutionPercentage</statString>
+		</graph>
+
+		<graph title="Physical Memory" smooth="0" thickness="4" maxY="0">
+			<statString>MemoryFreeMB,PhysicalUsedMB</statString>
+		</graph>
+
+		<graph title="RHI Drawcalls" maxY="4000" thickness="4">
+			<statString>RHI/Drawcalls</statString>
+		</graph>
+
+		<graph title="Gamethread Breakdown" showAverages="1" hideStatPrefix="exclusive/Gamethread/" stacked="1" mainStat="gamethreadtime" ignoreStats="exclusive/Gamethread/EventWait*" legendAverageThreshold="0.1">
+			<statString>exclusive/gamethread/*,gamethreadtime</statString>
+		</graph>
+
+		<graph title="Gamethread Breakdown Raw" hideStatPrefix="exclusive/Gamethread/" stacked="1" mainStat="gamethreadtime" ignoreStats="exclusive/Gamethread/EventWait*" smooth="0"  maxY="100" legendAverageThreshold="0.1">
+			<statString>exclusive/gamethread/*,gamethreadtime</statString>
+		</graph>
+
+		<graph title="Gamethread Waits" showAverages="1" hideStatPrefix="exclusive/gamethread/eventwait/" stacked="1" mainStat="frametime" ignoreStats="exclusive/Gamethread/EventWait" maxHierarchyDepth="3">
+			<statString>exclusive/gamethread/eventwait/*,gamethreadtime</statString>
+		</graph>
+
+		<graph title="Gamethread Physics" showAverages="1" stacked="1" maxY="0" maxHierarchyDepth="3" hideStatPrefix="exclusive/gamethread/">
+		  <statString>exclusive/gamethread/eventwait/EndPhysics,exclusive/gamethread/Physics,exclusive/gamethread/vehicle*</statString>
+		</graph>
+
+		<graph title="Gamethread Breakdown Giant" hideStatPrefix="exclusive/Gamethread/" stacked="1" mainStat="gamethreadtime" ignoreStats="exclusive/Gamethread/EventWait*" smooth="0"  maxY="100" legendAverageThreshold="0.1" width="40000" height="1000">
+			<statString>exclusive/gamethread/*,gamethreadtime</statString>
+		</graph>
+    
+		<graph title="Renderthread Breakdown" showAverages="1" hideStatPrefix="exclusive/RenderThread/" stacked="1" mainStat="renderthreadtime" ignoreStats="exclusive/Renderthread/FetchVisibilityForPrimitives;exclusive/renderthread/eventwait*">
+			<statString>exclusive/renderthread/*,renderthreadtime</statString>
+		</graph>
+
+		<graph title="Renderthread Breakdown Raw" hideStatPrefix="exclusive/RenderThread/" stacked="1" mainStat="renderthreadtime" ignoreStats="exclusive/Renderthread/FetchVisibilityForPrimitives;exclusive/renderthread/eventwait*" smooth="0"  maxY="100">
+			<statString>exclusive/renderthread/*,renderthreadtime</statString>
+		</graph>
+
+		<graph title="GPU Breakdown" showAverages="1" smoothKernelSize="120" hideStatPrefix="gpu/" stacked="1" mainStat="gpu/total">
+			<statString>GPU/*</statString>
+		</graph>
+
+		<graph title="Niagara Breakdown" showAverages="1" thickness="2" >
+			<statString>Ticks/NiagaraComponent,Exclusive/GameThread/Niagara,GPU/NiagaraGPUSimulation,GPU/NiagaraGPUSorting</statString>
+		</graph>
+
+		<graph title="Ticks Queued" showAverages="0" maxY="0" thickness="4" smooth="0">
+			<statString>Basic/TicksQueued</statString>
+		</graph>
+
+		<graph title="Ticks Breakdown" showAverages="1" maxY="0" maxHierarchyDepth="4" hideStatPrefix="Ticks/" stacked="1" mainStat="Basic/TicksQueued" threshold="1">
+			<statString>Ticks/*</statString>
+		</graph>
+
+		<graph title="Actor Counts" showAverages="0" maxY="0" thickness="4" hideStatPrefix="ActorCount/" smooth="0" stacked="1" mainStat="ActorCount/TotalActorCount">
+			<statString>ActorCount/*</statString>
+		</graph>
+
+		<graph title="Demo Recording" hideStatPrefix="basic/" maxY="15" thickness="5" showAverages="1">
+			<statString>Basic/GameThread/ConcurrentWithSlateTickT*,Basic/*/DemoRecordTime</statString>
+		</graph>
+
+		<graph title="Demo Recording (unsmoothed)" hideStatPrefix="basic/" maxY="0" thickness="5" showAverages="1" smooth="0">
+			<statString>Basic/GameThread/ConcurrentWithSlateTickT*,Basic/*/DemoRecordTime</statString>
+		</graph>
+
+		<graph title="Animation Budget" showAverages="0" hideStatPrefix="Exclusive/GameThread/" mainStat="Exclusive/GameThread/Animation" thickness="4" smooth="1" filterOutZeros="0">
+			<statString>Exclusive/GameThread/Animation,AnimationBudget/*</statString>
+		</graph>
+
+		<graph title="Log Counts" showAverages="1" hideStatPrefix="fmsglogf/" mainStat="fmsglogf/fmsglogfcount" thickness="4" smooth="0" filterOutZeros="0">
+			<statString>fmsglogf/*</statString>
+		</graph>
+
+		<graph title="CSV Profiler" showAverages="1" hideStatPrefix="CSVProfiler/" stacked="1" mainStat="frametime" maxY="0">
+			<statString>CSVProfiler/Num*</statString>
+		</graph>
+		
+		<graph title="Android Memory" showAverages="1" hideStatPrefix="AndroidMemory/" stacked="1" mainStat="MemoryFreeMB" maxY="0">
+			<statString>AndroidMemory/MemoryWarningState,AndroidMemory/TrimMemoryForegroundLevel,AndroidMemory/TrimMemoryBackgroundLevel</statString>
+		</graph>		
+
+		<graph title="Android Thermals" showAverages="1" hideStatPrefix="AndroidCPU/" maxY="100">
+			<statString>AndroidCPU/CPUTemp,AndroidCPU/ThermalStatus</statString>
+		</graph>
+		
+		<graph title="Android CPU" showAverages="1" hideStatPrefix="AndroidCPU/" maxY="0">
+			<statString>AndroidCPU/CPUFreqMHzGroup0,AndroidCPU/CPUFreqMHzGroup1,AndroidCPU/CPUFreqMHzGroup2,AndroidCPU/CPUFreqMHzGroup3</statString>
+		</graph>
+		
+		<graph title="Android CPU Percentage From Maximum" showAverages="1" hideStatPrefix="AndroidCPU/" maxY="0">
+			<statString>AndroidCPU/CPUFreqPercentageGroup0,AndroidCPU/CPUFreqPercentageGroup1,AndroidCPU/CPUFreqPercentageGroup2,AndroidCPU/CPUFreqPercentageGroup3</statString>
+		</graph>
+
+		<graph title="GPU cycles / frame" showAverages="1" hideStatPrefix="AndroidGPU/" mainStat="frametime" maxY="0">
+			<statString>AndroidGPU/GPUCyclesMln,AndroidGPU/VertexCyclesMln,AndroidGPU/FragmentCyclesMln</statString>
+		</graph>
+		
+		<graph title="Pixels / frame" showAverages="1" hideStatPrefix="AndroidGPU/" maxY="0">
+			<statString>AndroidGPU/PixelsMln</statString>
+		</graph>
+		
+		<graph title="Shader cycles / frame" showAverages="1" hideStatPrefix="AndroidGPU/" maxY="0">
+			<statString>AndroidGPU/ShaderCyclesMln,AndroidGPU/ShaderArithmeticCyclesMln,AndroidGPU/ShaderLoadStoreCyclesMln,AndroidGPU/ShaderTextureCyclesMln</statString>
+		</graph>
+		
+		<graph title="GPU bandwidth / frame" showAverages="1" hideStatPrefix="AndroidGPU/" stacked="1" maxY="0">
+			<statString>AndroidGPU/ExternalMemoryReadMB,AndroidGPU/ExternalMemoryWriteMB</statString>
+		</graph>
+
+	</graphGroup>
+
+	<graphGroup name="AI5">		
+		<baseSettings smooth="0" maxY="0"/>
+
+		<graph title="All StateTree Costs" showAverages="1" hideStatPrefix="Exclusive/Gamethread/" stacked="1">
+			<statString>Exclusive/Gamethread/StateTree*</statString>
+		</graph>
+
+		<graph title="StateTree Processor Costs" showAverages="1" hideStatPrefix="Exclusive/Gamethread/StateTreeProcessor" stacked="1">
+			<statString>Exclusive/Gamethread/StateTreeProcessor*</statString>
+		</graph>
+
+		<graph title="StateTree Context Costs" showAverages="1" hideStatPrefix="Exclusive/Gamethread/StateTree_" stacked="1">
+			<statString>Exclusive/Gamethread/StateTree_*</statString>
+		</graph>
+
+		<graph title="Ticked StateTree Count" showAverages="1" hideStatPrefix="StateTreeProcessor/">
+			<statString>StateTreeProcessor/NumTickedStateTree</statString>
+		</graph>
+
+		<graph title="Mass Signals Count" showAverages="1" hideStatPrefix="MassSignalsCounters/">
+			<statString>MassSignalsCounters/*</statString>
+		</graph>
+		
+		<graph title="Deferred Commands Count" showAverages="1" hideStatPrefix="MassEntitiesCounters/">
+			<statString>MassEntitiesCounters/*</statString>
+		</graph>
+		
+		<graph title="Deferred Commands Costs" showAverages="1" hideStatPrefix="MassEntities/Gamethread/" stacked="1">
+			<statString>MassEntities/*</statString>
+		</graph>
+
+		<graph title="Mass Actor Counts" showAverages="1" hideStatPrefix="MassActors/" stacked="1">
+			<statString>MassActors/*</statString>
+		</graph>
+
+	</graphGroup>
+
+</graphGroups>
diff --git a/Engine/Source/Programs/CSVTools/PerfReportTool/ReportTypes.xml b/Engine/Source/Programs/CSVTools/PerfReportTool/ReportTypes.xml
new file mode 100644
--- /dev/null
+++ b/Engine/Source/Programs/CSVTools/PerfReportTool/ReportTypes.xml
@@ -0,0 +1,192 @@
+<?xml version="1.0" encoding="UTF-8"?>
+
+<!-- Please use the XML schemas in "Engine/Source/Programs/CsvTools/PerfReportTool/XmlSchemas" to validate -->
+<root xmlns="https://www.unrealengine.com/PerfReportTool" xmlns:xsi="http://www.w3.org/2001/XMLSchema-instance">
+	<!-- Friendly names for display purposes -->
+	<statDisplayNameMappings>
+		<mapping statName="RHI/Drawcalls" displayName="DrawCalls"/>
+		<mapping statName="DynamicResolutionPercentage" displayName="DynRes %"/>
+		<mapping statName="Basic/TicksQueued" displayName="TicksQueued"/>
+		<mapping statName="MemoryFreeMB" displayName="MemFreeMB"/>
+		<mapping statName="PhysicalUsedMB" displayName="MemUsedMB"/> 
+		<mapping statName="buildversion" displayName="Build Version"/>
+		<mapping statName="config" displayName="Config"/>
+		<mapping statName="commandline" displayName="Commandline"/>
+		<mapping statName="syncLoads" displayName="Sync Loads"/>
+		<mapping statName="scalabilitylevel" displayName="Scaling"/>
+		<mapping statName="deviceprofile" displayName="DeviceProfile"/>
+		<mapping statName="TargetFPS" displayName="FPS"/>
+		<mapping statName="HitchTimePercent" displayName="HitchTime %"/>
+	</statDisplayNameMappings>
+
+	<!-- Strip Events - these delimit frames which will be excluded from stats computation -->
+	<!-- Note: order is important, since these are stripped in the order listed (with no handling for overlapping events) -->
+	<csvEventsToStrip>
+		<eventPair begin="App_WillDeactivate" end="App_HasReactivated"/>
+	</csvEventsToStrip>
+
+	<!-- Preset summary tables -->
+	<!-- In bulk mode (-csvdir), summary table row data is produced for each CSV. This includes data from all active summaries -->
+	<!-- The preset Summary Tables below can be used to provide a filtered and sorted view of that data (see -summaryTable). -->
+	<summaryTables>
+		<summaryTable name="all" rowSort="buildversion,Platform,deviceprofile">
+			<filter>buildversion,Platform,DeviceProfile,scalabilitylevel,*</filter>
+		</summaryTable>
+
+		<summaryTable name="default" rowSort="buildversion,playlistid,Platform,deviceprofile,scalabilitylevel,TargetFPS">
+			<filter>Platform,DeviceProfile,scalabilitylevel,Report,TargetFPS,MVP*,Hitches/Min,HitchTimePercent,FrameTime*,GamethreadTime*,Renderthreadtime*,RHIThreadTime*,GPUTime*,MemoryFreeMB Min,PhysicalUsedMB Max,DynamicResolution*,RHI/Drawcalls*,Basic/TicksQueued*,Total Time*,Csv File,buildversion,config,playlistid,framecount*,syncloads</filter>
+		</summaryTable>
+
+		<summaryTable name="memory" rowSort="buildversion,Platform,deviceprofile,scalabilitylevel">
+			<filter>MemoryFreeMB*,PhysicalUsedMB*,Total Time*</filter>
+		</summaryTable>
+
+		<summaryTable name="condensed" rowSort="buildversion,Platform,deviceprofile,scalabilitylevel,TargetFPS">
+			<filter>DeviceProfile,scalabilitylevel,TargetFPS,MVP*,Hitches/Min,MemoryFreeMB Min,PhysicalUsedMB Max,GamethreadTime*,Renderthreadtime*,RHIThreadTime*,GPUTime*,DynamicResolution*,SyncLoads*</filter>
+		</summaryTable>
+	</summaryTables>
+
+	<reporttypes reportGraphsFile="ReportGraphs.xml">
+
+		<!-- 60fps report -->
+		<reporttype name="Default60fps" title="60FPS Performance Report" vsync="1">
+			<autodetection requiredstats="frametime">
+				<requiredmetadata key="targetframerate" allowedValues="60" stopIfKeyFound="1" ignoreIfKeyNotFound="1"/>
+			</autodetection>
+			
+			<metadataToShow>buildversion,platform,deviceprofile</metadataToShow>
+
+			<summary type="fpschart" fps="60" hitchThreshold="60" useEngineHitchMetric="1">
+				<stats>Frametime,GameThreadtime,RenderThreadtime,RHIThreadTime,GPUtime,RHI/Drawcalls,DynamicResolutionPercentage,MemoryFreeMB(min),PhysicalUsedMB(max),Basic/TicksQueued</stats>
+				<colourThresholds stat="Hitches/Min"> 2, 3, 5, 6</colourThresholds>
+				<colourThresholds stat="HitchTimePercent"> 0.1, 0.5, 1, 2</colourThresholds>
+				<colourThresholds stat="MVP60"> 2, 3, 4, 5</colourThresholds>
+				<colourThresholds stat="MemoryFreeMB(min)">400,300,200,50</colourThresholds>
+				<colourThresholds stat="DynamicResolutionPercentage">90,80,70,60</colourThresholds>
+				<colourThresholds stat="Frametime"> 16.7, 17.5, 19, 22.5 </colourThresholds>
+				<colourThresholds stat="GameThreadtime"> 16.7, 17.5, 19, 22.5 </colourThresholds>
+				<colourThresholds stat="RHIThreadTime"> 16.7, 17.5, 19, 22.5 </colourThresholds>
+				<colourThresholds stat="RenderThreadtime"> 16.7, 17.5, 19, 22.5 </colourThresholds>
+				<colourThresholds stat="GPUtime"> 16.7, 17.5, 19, 22.5 </colourThresholds>
+				<colourThresholds stat="Basic/TicksQueued"> 150, 200, 250, 300 </colourThresholds>
+			</summary>
+
+			<summary type="hitches">
+				<stats>Frametime,GameThreadtime,RenderThreadtime,RHIThreadTime,GPUtime</stats>
+				<colourThresholds>33, 20, 16.6, 16.2</colourThresholds>
+				<hitchThresholds>60, 150, 250, 500, 750, 1000, 2000</hitchThresholds>
+			</summary>
+
+			<graph title="Stat Unit" budget="16.66" />
+			<graph title="Stat Unit Raw" budget="16.66" />
+			<graph title="Gamethread Breakdown" budget="16.66"/>
+			<graph title="Gamethread Breakdown Raw" budget="16.66"/>
+			<graph title="Gamethread Waits" budget="2"/>
+			<graph title="Gamethread Physics" budget="2"/>
+			<graph title="Ticks Breakdown" budget="200"/>
+			<graph title="Actor Counts" budget="0"/>
+			<graph title="Physical Memory" budget="300"/>
+			<graph title="Dynamic Resolution" budget="80"/>
+			<graph title="GPU Breakdown" budget="16.66" />
+			<graph title="Niagara Breakdown" budget="33.33"/>
+			<graph title="RHI Drawcalls" budget="3000" />
+			<graph title="Renderthread Breakdown" budget="16.66"/>
+			<graph title="Renderthread Breakdown Raw" budget="16.66"/>
+			<graph title="Demo Recording" budget="16.66" />
+			<graph title="Demo Recording (unsmoothed)" budget="16.66" />
+			<graph title="Animation Budget" budget="1.5" />
+			<graph title="Log Counts" budget="0" />
+			<graph title="CSV Profiler" budget="0" />
+			<graph title="Android Memory" budget="0" />
+			<graph title="Android Thermals" budget="0" />
+			<graph title="Android CPU" budget="0" />
+			<graph title="Android CPU Percentage From Maximum" budget="0" />
+			<graph title="GPU cycles / frame" budget="0" />
+			<graph title="Pixels / frame" budget="0" />
+			<graph title="Shader cycles / frame" budget="0" />
+			<graph title="GPU bandwidth / frame" budget="0" />
+		</reporttype>
+
+
+		<!-- 30 Fps report -->
+		<reporttype name="Default30fps" title="30FPS Performance Report" vsync="1">
+			<autodetection requiredstats="frametime">
+				<requiredmetadata key="targetframerate" allowedValues="30" stopIfKeyFound="1" ignoreIfKeyNotFound="1"/>
+			</autodetection>
+
+			<metadataToShow>buildversion,platform,deviceprofile</metadataToShow>
+
+			<summary type="fpschart" fps="30" hitchThreshold="60" useEngineHitchMetric="1">
+				<stats>Frametime,GameThreadtime,RenderThreadtime,RHIThreadTime,GPUtime,RHI/Drawcalls,DynamicResolutionPercentage,MemoryFreeMB(min),PhysicalUsedMB(max),Basic/TicksQueued</stats>
+				<colourThresholds stat="Hitches/Min"> 0, 2.5, 3.75, 5</colourThresholds>
+				<colourThresholds stat="HitchTimePercent"> 0.1, 0.5, 1, 2</colourThresholds>
+				<colourThresholds stat="MVP30"> 0, 2, 3, 4</colourThresholds>
+				<colourThresholds stat="MemoryFreeMB">400,300,200,100</colourThresholds>
+				<colourThresholds stat="MemoryFreeMB(min)">400,300,200,50</colourThresholds>
+				<colourThresholds stat="DynamicResolutionPercentage">90,80,70,60</colourThresholds>
+				<colourThresholds stat="Frametime"> 33.4, 35, 38, 45  </colourThresholds>
+				<colourThresholds stat="GameThreadtime"> 33.4, 35, 38, 45  </colourThresholds>
+				<colourThresholds stat="RHIThreadTime"> 33.4, 35, 38, 45  </colourThresholds>
+				<colourThresholds stat="RenderThreadtime"> 33.4, 35, 38, 45  </colourThresholds>
+				<colourThresholds stat="GPUtime"> 33.4, 35, 38, 45  </colourThresholds>
+				<colourThresholds stat="Basic/TicksQueued"> 150, 200, 250, 300 </colourThresholds>
+			</summary>
+
+			<summary type="hitches">
+				<stats>Frametime,GameThreadtime,RenderThreadtime,RHIThreadTime,GPUtime</stats>
+				<colourThresholds> 20, 16.6, 16.2</colourThresholds>
+				<hitchThresholds>60, 150, 250, 500, 750, 1000, 2000</hitchThresholds>
+			</summary>
+
+			<graph title="Stat Unit 30hz" budget="33.33" />
+			<graph title="Stat Unit Raw" budget="33.33" />
+			<graph title="Gamethread Breakdown" budget="33.33"/>
+			<graph title="Gamethread Breakdown Raw" budget="33.33"/>
+			<graph title="Gamethread Waits" budget="4"/>
+			<graph title="Gamethread Physics" budget="4"/>
+			<graph title="Ticks Breakdown" budget="200"/>
+			<graph title="Actor Counts" budget="0"/>
+			<graph title="Physical Memory" budget="300"/>
+			<graph title="Dynamic Resolution" budget="80"/>
+			<graph title="GPU Breakdown" budget="33.33" />
+			<graph title="Niagara Breakdown" budget="33.33"/>
+			<graph title="RHI Drawcalls" budget="2000" />
+			<graph title="Renderthread Breakdown" budget="33.33"/>
+			<graph title="Renderthread Breakdown Raw" budget="33.33"/>
+			<graph title="Animation Budget" budget="3.0" />
+			<graph title="Log Counts" budget="0" />
+			<graph title="CSV Profiler" budget="0" />
+			<graph title="Android Memory" budget="0" />
+			<graph title="Android Thermals" budget="0" />
+			<graph title="Android CPU" budget="0" />
+			<graph title="Android CPU Percentage From Maximum" budget="0" />
+			<graph title="GPU cycles / frame" budget="0" />
+			<graph title="Pixels / frame" budget="0" />
+			<graph title="Shader cycles / frame" budget="0" />
+			<graph title="GPU bandwidth / frame" budget="0" />
+		</reporttype>
+
+		<!-- AI5 Fps report -->
+		<reporttype name="AI5" title="AI Performance Report" vsync="1">
+			<metadataToShow>buildversion,platform,deviceprofile</metadataToShow>
+
+			<graph title="All StateTree Costs" budget="0"/>
+			<graph title="StateTree Processor Costs" budget="0"/>
+			<graph title="StateTree Context Costs" budget="0"/>
+			<graph title="Ticked StateTree Count" budget="0"/>
+			<graph title="Mass Signals Count" budget="0"/>
+			<graph title="Deferred Commands Count" budget="0"/>
+			<graph title="Deferred Commands Costs" budget="0"/>
+			<graph title="Mass Actor Counts" budget="0" />
+		</reporttype>
+
+		<reporttype name="Memory" title="Physical Memory Report" vsync="1">
+			<autodetection requiredstats="memoryfreemb,physicalusedmb"/>
+			<summary type="fpschart" fps="60" hitchThreshold="60" useEngineHitchMetric="1">
+				<stats>MemoryFreeMB(min),PhysicalUsedMB(max),MemoryFreeMB,PhysicalUsedMB,PhysicalUsedMB(min)</stats>
+				<colourThresholds stat="MemoryFreeMB(min)">400,300,200,50</colourThresholds>
+			</summary>
+			<graph title="Physical Memory" budget="300"/>
+		</reporttype>
+	</reporttypes>
+</root>
\ No newline at end of file
diff --git a/Engine/Source/Programs/CSVTools/PerfReportTool/Summaries/MapOverlaySummary.cs b/Engine/Source/Programs/CSVTools/PerfReportTool/Summaries/MapOverlaySummary.cs
--- a/Engine/Source/Programs/CSVTools/PerfReportTool/Summaries/MapOverlaySummary.cs
+++ b/Engine/Source/Programs/CSVTools/PerfReportTool/Summaries/MapOverlaySummary.cs
@@ -208,7 +208,7 @@ namespace PerfSummaries
 
 		private void CopyAndResizeImage(string sourceImagePath, string destImagePath, int destWidth, int destHeight)
 		{
-			if (!OperatingSystem.IsWindows())
+			if (!OperatingSystem.IsWindowsVersionAtLeast(6, 1))
 			{
 				Console.WriteLine("CopyAndResizeImage is not supported on this platform!");
 				return;
@@ -232,7 +232,7 @@ namespace PerfSummaries
 
 		private void CopyAndTrimImage(string sourceImagePath, string destImagePath, int destWidth, int destHeight, Rectangle subRegion)
 		{
-			if (!OperatingSystem.IsWindows())
+			if (!OperatingSystem.IsWindowsVersionAtLeast(6, 1))
 			{
 				Console.WriteLine("CopyAndTrimImage is not supported on this platform!");
 				return;
diff --git a/Engine/Source/Programs/CSVTools/Publish.bat b/Engine/Source/Programs/CSVTools/Publish.bat
new file mode 100755
--- /dev/null
+++ b/Engine/Source/Programs/CSVTools/Publish.bat
@@ -0,0 +1,57 @@
+@echo off
+
+rem ## Unreal Engine UnrealBuildTool build script
+rem ## Copyright Epic Games, Inc. All Rights Reserved.
+
+setlocal
+
+pushd "%~dp0
+
+rem ## Verify that dotnet is present
+call "%~dp0..\..\..\Build\BatchFiles\GetDotnetPath.bat"
+if errorlevel 1 goto Error_NoDotnetSDK
+
+rem ## Command line arguments
+set MSBUILD_LOGLEVEL=%1
+if not defined %MSBUILD_LOGLEVEL set MSBUILD_LOGLEVEL=quiet
+
+echo Building CsvTools...
+dotnet build CsvTools.sln -c Release -v %MSBUILD_LOGLEVEL%
+if errorlevel 1 goto Error_CompileFailed
+
+echo Publishing CsvTools...
+if exist ..\..\..\Binaries\DotNET\CsvTools rmdir /s /q ..\..\..\Binaries\DotNET\CsvTools >nul 2>&1
+if not exist ..\..\..\Binaries\DotNET\CsvTools mkdir ..\..\..\Binaries\DotNET\CsvTools >nul 2>&1
+dotnet publish CSVCollate\CSVCollate.csproj -c Release --output ..\..\..\Binaries\DotNET\CsvTools --no-build -v %MSBUILD_LOGLEVEL%
+dotnet publish CsvConvert\CsvConvert.csproj -c Release --output ..\..\..\Binaries\DotNET\CsvTools --no-build -v %MSBUILD_LOGLEVEL%
+dotnet publish CSVFilter\CSVFilter.csproj -c Release --output ..\..\..\Binaries\DotNET\CsvTools --no-build -v %MSBUILD_LOGLEVEL%
+dotnet publish CsvInfo\CsvInfo.csproj -c Release --output ..\..\..\Binaries\DotNET\CsvTools --no-build -v %MSBUILD_LOGLEVEL%
+dotnet publish CSVSplit\CSVSplit.csproj -c Release --output ..\..\..\Binaries\DotNET\CsvTools --no-build -v %MSBUILD_LOGLEVEL%
+dotnet publish CsvStats\CsvStats.csproj -c Release --output ..\..\..\Binaries\DotNET\CsvTools --no-build -v %MSBUILD_LOGLEVEL%
+dotnet publish CSVToSVG\CSVToSVG.csproj -c Release --output ..\..\..\Binaries\DotNET\CsvTools --no-build -v %MSBUILD_LOGLEVEL%
+dotnet publish CsvToSvgLib\CsvToSvgLib.csproj -c Release --output ..\..\..\Binaries\DotNET\CsvTools --no-build -v %MSBUILD_LOGLEVEL%
+dotnet publish PerfReportTool\PerfreportTool.csproj -c Release --output ..\..\..\Binaries\DotNET\CsvTools --no-build -v %MSBUILD_LOGLEVEL%
+dotnet publish RegressionsReport\RegressionsReport.csproj -c Release --output ..\..\..\Binaries\DotNET\CsvTools --no-build -v %MSBUILD_LOGLEVEL%
+
+rem ## Success!
+set PUBLISH_EXITCODE=0
+goto Exit
+
+:Error_NoDotnetSDK
+echo.
+echo Publish ERROR: Unable to find a install of Dotnet SDK.  Please make sure you have it installed and that `dotnet` is a globally available command.
+echo.
+set PUBLISH_EXITCODE=1
+goto Exit
+
+:Error_CompileFailed
+echo.
+echo Publish ERROR: CsvTools failed to compile.
+echo.
+set PUBLISH_EXITCODE=1
+goto Exit
+
+:Exit
+rem ## Restore original CWD in case we change it
+popd
+exit /B %PUBLISH_EXITCODE%
diff --git a/Engine/Source/Programs/CSVTools/PublishWithP4.bat b/Engine/Source/Programs/CSVTools/PublishWithP4.bat
new file mode 100755
--- /dev/null
+++ b/Engine/Source/Programs/CSVTools/PublishWithP4.bat
@@ -0,0 +1,44 @@
+@echo off
+setlocal enabledelayedexpansion
+
+call .\Publish.bat
+if errorlevel 1 goto Error
+
+set "BinariesDir=%~dp0..\..\..\Binaries\DotNET\CsvTools"
+echo Binaries directory: %BinariesDir%
+
+if not exist "%BinariesDir%" (
+    echo CsvTools binaries folder not found:
+    echo %BinariesDir%
+    goto exit
+)
+
+pushd "%~dp0\..\..\..\Binaries\DotNET\CsvTools
+
+for /f "tokens=1" %%A in ('p4 where . ^| findstr /c:"//"') do (
+    set p4Path=%%A
+)
+echo P4 Path: !p4Path!
+
+p4 edit !p4Path!/...
+p4 --field "Description=CSV Tools Binaries" --field "Files=" change -o | p4 change -i > tempP4Changeoutput.txt
+
+REM Read the CL
+For /F "delims=" %%L IN (tempP4Changeoutput.txt) DO (
+  set "line=%%L"
+  set "change=!line:Change =!"
+  set "change=!change: created.=!"
+)
+echo Creating new CL: !change!
+del tempP4ChangeOutput.txt
+
+p4 reopen -c !change! !p4Path!/...
+
+goto exit
+
+:error
+
+:exit
+EndLocal
+popd
+
diff --git a/Engine/Source/Programs/CSVTools/RegressionsReport/Properties/AssemblyInfo.cs b/Engine/Source/Programs/CSVTools/RegressionsReport/Properties/AssemblyInfo.cs
--- a/Engine/Source/Programs/CSVTools/RegressionsReport/Properties/AssemblyInfo.cs
+++ b/Engine/Source/Programs/CSVTools/RegressionsReport/Properties/AssemblyInfo.cs
@@ -1,13 +1,6 @@
 // Copyright (C) Microsoft. All rights reserved.
 // Copyright Epic Games, Inc. All Rights Reserved.
-using System.Reflection;
-using System.Runtime.CompilerServices;
 using System.Runtime.InteropServices;
-[assembly: AssemblyTrademark("")]
-[assembly: AssemblyCulture("")]
-// Setting ComVisible to false makes the types in this assembly not visible 
-// to COM components.  If you need to access a type in this assembly from 
-// COM, set the ComVisible attribute to true on that type.
-[assembly: ComVisible(false)]
+
 // The following GUID is for the ID of the typelib if this project is exposed to COM
 [assembly: Guid("06104e7c-f154-4864-af94-996d8bba6319")]
\ No newline at end of file
diff --git a/Engine/Source/Programs/CSVTools/RegressionsReport/RegressionsReport.csproj b/Engine/Source/Programs/CSVTools/RegressionsReport/RegressionsReport.csproj
--- a/Engine/Source/Programs/CSVTools/RegressionsReport/RegressionsReport.csproj
+++ b/Engine/Source/Programs/CSVTools/RegressionsReport/RegressionsReport.csproj
@@ -1,16 +1,8 @@
 <Project Sdk="Microsoft.NET.Sdk">
   <PropertyGroup>
-    <TargetFramework>net6.0</TargetFramework>
     <OutputType>Exe</OutputType>
-    <GenerateAssemblyInfo>false</GenerateAssemblyInfo>
-    <AppendTargetFrameworkToOutputPath>false</AppendTargetFrameworkToOutputPath>
-  </PropertyGroup>
-  <PropertyGroup Condition=" '$(Configuration)|$(Platform)' == 'Debug|AnyCPU' ">
-    <PlatformTarget>x64</PlatformTarget>
-  </PropertyGroup>
-  <PropertyGroup Condition=" '$(Configuration)|$(Platform)' == 'Release|AnyCPU' ">
-    <PlatformTarget>x64</PlatformTarget>
-    <OutputPath>..\..\..\..\Binaries\DotNET\CsvTools\</OutputPath>
+    <AssemblyTitle>RegressionsReport</AssemblyTitle>
+    <Product>RegressionsReport</Product>
   </PropertyGroup>
   <ItemGroup>
     <None Remove="Properties\AssemblyInfo.cs~RF273305d.TMP" />
@@ -19,16 +11,7 @@
     <ProjectReference Include="..\CsvStats\CsvStats.csproj" />
   </ItemGroup>
   <ItemGroup>
-    <PackageReference Include="System.Drawing.Common" Version="7.0.0" />
-    <PackageReference Include="Microsoft.CSharp" Version="4.7.0" />
+    <PackageReference Include="System.Drawing.Common" Version="8.0.24" />
     <PackageReference Include="System.Data.DataSetExtensions" Version="4.5.0" />
   </ItemGroup>
-  <PropertyGroup>
-    <AssemblyTitle>RegressionsReport</AssemblyTitle>
-    <Product>RegressionsReport</Product>
-    <Copyright>Copyright %28C%29 Microsoft. All rights reserved.
-Copyright Epic Games, Inc. All Rights Reserved.</Copyright>
-    <AssemblyVersion>1.0.0.0</AssemblyVersion>
-    <FileVersion>1.0.0.0</FileVersion>
-  </PropertyGroup>
 </Project>
