CSV Tools PublishWithP4.bat: reconcile rather than edit so we clear unneeded files and add new ones

diff --git a/Engine/Source/Programs/CSVTools/PublishWithP4.bat b/Engine/Source/Programs/CSVTools/PublishWithP4.bat
--- a/Engine/Source/Programs/CSVTools/PublishWithP4.bat
+++ b/Engine/Source/Programs/CSVTools/PublishWithP4.bat
@@ -2,7 +2,7 @@
 setlocal enabledelayedexpansion
 
 call .\Publish.bat
-if errorlevel 1 goto Error
+if errorlevel 1 goto exit
 
 set "BinariesDir=%~dp0..\..\..\Binaries\DotNET\CsvTools"
 echo Binaries directory: %BinariesDir%
@@ -20,7 +20,12 @@ for /f "tokens=1" %%A in ('p4 where . ^| findstr /c:"//"') do (
 )
 echo P4 Path: !p4Path!
 
-p4 edit !p4Path!/...
+p4 reconcile !p4Path!/...
+if errorlevel 1 (
+	echo Error with p4 reconcile
+	goto exit
+)
+
 p4 --field "Description=CSV Tools Binaries" --field "Files=" change -o | p4 change -i > tempP4Changeoutput.txt
 
 REM Read the CL
@@ -33,10 +38,12 @@ echo Creating new CL: !change!
 del tempP4ChangeOutput.txt
 
 p4 reopen -c !change! !p4Path!/...
+if errorlevel 1 (
+	echo Error with p4 reopen
+	goto exit
+)
 
-goto exit
-
-:error
+echo Success
 
 :exit
 EndLocal
