CsvTools: PublishWithP4: fix a mismatched popd in an error case

diff --git a/Engine/Source/Programs/CSVTools/PublishWithP4.bat b/Engine/Source/Programs/CSVTools/PublishWithP4.bat
--- a/Engine/Source/Programs/CSVTools/PublishWithP4.bat
+++ b/Engine/Source/Programs/CSVTools/PublishWithP4.bat
@@ -2,7 +2,7 @@
 setlocal enabledelayedexpansion
 
 call .\Publish.bat
-if errorlevel 1 goto exit
+if errorlevel 1 goto publishError
 
 set "BinariesDir=%~dp0..\..\..\Binaries\DotNET\CsvTools"
 echo Binaries directory: %BinariesDir%
@@ -46,6 +46,7 @@ if errorlevel 1 (
 echo Success
 
 :exit
-EndLocal
 popd
+:publishError
+EndLocal
 
