PerfReportTool 4.269.0: Add anchor links for summary headings
#Synced-CL 46049276

diff --git a/Engine/Source/Programs/CSVTools/PerfReportTool/PerfReportTool.cs b/Engine/Source/Programs/CSVTools/PerfReportTool/PerfReportTool.cs
--- a/Engine/Source/Programs/CSVTools/PerfReportTool/PerfReportTool.cs
+++ b/Engine/Source/Programs/CSVTools/PerfReportTool/PerfReportTool.cs
@@ -24,7 +24,7 @@ namespace PerfReportTool
 	class Version
 	{
 		// Format: Major.Minor.Bugfix
-		private static string VersionString = "4.268.1";
+		private static string VersionString = "4.269.0";
 
 		public static string Get() { return VersionString; }
 	};
@@ -1623,6 +1623,10 @@ namespace PerfReportTool
 				htmlFile.WriteLine("      .collapsibleSection.expanded { grid-template-rows: 1fr; }");
 				htmlFile.WriteLine("      .collapsibleSectionInner { overflow: auto; overflow-y: hidden; }");
 
+				// Section anchors
+				htmlFile.WriteLine("      .sectionAnchorLink {margin-left:auto; position:relative; top:32px}");
+				htmlFile.WriteLine("      .sectionAnchorDiv {display:flex; align-items:baseline}");
+
 				htmlFile.WriteLine("    </style>");
 				htmlFile.WriteLine("  </head>");
 				htmlFile.WriteLine("  <body>");
@@ -1673,6 +1677,7 @@ namespace PerfReportTool
 				htmlFile.WriteLine("<button type='button' onclick='expandAll()' style='margin-right:10px'>Expand all</button>");
 				htmlFile.WriteLine("<button type='button' onclick='location.href=\"#graphList\";'>Graphs</button>");
 				htmlFile.WriteLine("</div>");
+				htmlFile.WriteLine("<br>");
 			}
 
 			if (summaryRowData != null)
@@ -1696,7 +1701,13 @@ namespace PerfReportTool
 			// Add the standard summary last
 			summaries.Add(new StandardSummary());
 
-			// If the reporttype has summary info, then write out the summary]
+			// Assign indices to the final summary list
+			for (int i=0; i<summaries.Count; i++)
+			{
+				summaries[i].SetReportSummaryIndex(i);
+			}
+
+				// If the reporttype has summary info, then write out the summary]
 			foreach (Summary summary in summaries)
 			{
 				bool bWriteHtml = htmlFile != null && !summary.bHideInDetailedReport;
@@ -1750,7 +1761,7 @@ namespace PerfReportTool
 					ReportGraph graph = csvSvgInfo.Graph;
 
 					string svgTitle = graph.title;
-					HtmlSection htmlSection = new HtmlSection(svgTitle, false, StripSpaces(svgTitle), 3);
+					HtmlSection htmlSection = new HtmlSection(svgTitle, false, StripSpaces(svgTitle), 3, graph.settings.width.value);
 
 					if (csvSvgInfo.Format == CsvSvgInfo.GraphFormat.Inline)
 					{
diff --git a/Engine/Source/Programs/CSVTools/PerfReportTool/Summaries/BoundedStatValuesSummary.cs b/Engine/Source/Programs/CSVTools/PerfReportTool/Summaries/BoundedStatValuesSummary.cs
--- a/Engine/Source/Programs/CSVTools/PerfReportTool/Summaries/BoundedStatValuesSummary.cs
+++ b/Engine/Source/Programs/CSVTools/PerfReportTool/Summaries/BoundedStatValuesSummary.cs
@@ -1,11 +1,12 @@
 // Copyright Epic Games, Inc. All Rights Reserved.
 
+using CSVStats;
+using PerfReportTool;
 using System;
 using System.Collections.Generic;
 using System.Linq;
+using System.Text.RegularExpressions;
 using System.Xml.Linq;
-using PerfReportTool;
-using CSVStats;
 
 namespace PerfSummaries
 {
@@ -464,7 +465,7 @@ namespace PerfSummaries
 			// Output HTML
 			if (bWriteHtml)
 			{
-				htmlSection = new HtmlSection(title, bStartCollapsed);
+				htmlSection = new HtmlSection(title, bStartCollapsed, MakeSectionId(title));
 
 				htmlSection.WriteLine("  <table border='0' style='width:1400'>");
 				htmlSection.WriteLine("  <tr>");
diff --git a/Engine/Source/Programs/CSVTools/PerfReportTool/Summaries/BucketSummary.cs b/Engine/Source/Programs/CSVTools/PerfReportTool/Summaries/BucketSummary.cs
--- a/Engine/Source/Programs/CSVTools/PerfReportTool/Summaries/BucketSummary.cs
+++ b/Engine/Source/Programs/CSVTools/PerfReportTool/Summaries/BucketSummary.cs
@@ -91,7 +91,7 @@ namespace PerfSummaries
 				return null;
 			}
 
-			HtmlSection htmlSection = new HtmlSection(Title, bStartCollapsed);
+			HtmlSection htmlSection = new HtmlSection(Title, bStartCollapsed, MakeSectionId(Title));
 
 			htmlSection.WriteLine("  <table border='0' style='width:1000'>");
 			htmlSection.WriteLine("  <tr><td></td>");
diff --git a/Engine/Source/Programs/CSVTools/PerfReportTool/Summaries/CheckpointSummary.cs b/Engine/Source/Programs/CSVTools/PerfReportTool/Summaries/CheckpointSummary.cs
--- a/Engine/Source/Programs/CSVTools/PerfReportTool/Summaries/CheckpointSummary.cs
+++ b/Engine/Source/Programs/CSVTools/PerfReportTool/Summaries/CheckpointSummary.cs
@@ -93,7 +93,7 @@ namespace PerfSummaries
 			if (bWriteHtml)
 			{
 				// write out all the headers to the HTML
-				htmlSection = new HtmlSection(summaryTitle, bStartCollapsed);
+				htmlSection = new HtmlSection(summaryTitle, bStartCollapsed, MakeSectionId(summaryTitle));
 				string HeaderRow = "";
 
 				HeaderRow += "<th>Checkpoint Name</th>";
diff --git a/Engine/Source/Programs/CSVTools/PerfReportTool/Summaries/EventSummary.cs b/Engine/Source/Programs/CSVTools/PerfReportTool/Summaries/EventSummary.cs
--- a/Engine/Source/Programs/CSVTools/PerfReportTool/Summaries/EventSummary.cs
+++ b/Engine/Source/Programs/CSVTools/PerfReportTool/Summaries/EventSummary.cs
@@ -72,7 +72,7 @@ namespace PerfSummaries
 			// Output HTML
 			if (bWriteHtml && eventCountsDict.Count > 0)
 			{
-				htmlSection = new HtmlSection(title, bStartCollapsed);
+				htmlSection = new HtmlSection(title, bStartCollapsed, MakeSectionId(title));
 				htmlSection.WriteLine("  <table border='0' style='width:1200'>");
 				htmlSection.WriteLine("  <tr><th>Name</th><th><b>Count</th></tr>");
 				foreach (KeyValuePair<string, int> pair in eventCountsDict.ToList())
@@ -85,7 +85,6 @@ namespace PerfSummaries
 			// Output summary table row data
 			if (rowData != null)
 			{
-
 				ColourThresholdList thresholdList = null;
 
 				if (colourThresholds != null)
diff --git a/Engine/Source/Programs/CSVTools/PerfReportTool/Summaries/ExtraLinksSummary.cs b/Engine/Source/Programs/CSVTools/PerfReportTool/Summaries/ExtraLinksSummary.cs
--- a/Engine/Source/Programs/CSVTools/PerfReportTool/Summaries/ExtraLinksSummary.cs
+++ b/Engine/Source/Programs/CSVTools/PerfReportTool/Summaries/ExtraLinksSummary.cs
@@ -117,7 +117,7 @@ namespace PerfSummaries
 			// Output HTML
 			if (bWriteHtml)
 			{
-				htmlSection = new HtmlSection(title, bStartCollapsed);
+				htmlSection = new HtmlSection(title, bStartCollapsed, MakeSectionId(title));
 				htmlSection.WriteLine("  <ul>");
 				foreach (ExtraLink link in links)
 				{
diff --git a/Engine/Source/Programs/CSVTools/PerfReportTool/Summaries/FPSChartSummary.cs b/Engine/Source/Programs/CSVTools/PerfReportTool/Summaries/FPSChartSummary.cs
--- a/Engine/Source/Programs/CSVTools/PerfReportTool/Summaries/FPSChartSummary.cs
+++ b/Engine/Source/Programs/CSVTools/PerfReportTool/Summaries/FPSChartSummary.cs
@@ -302,7 +302,7 @@ namespace PerfSummaries
 			// Output HTML
 			if (bWriteHtml)
 			{
-				htmlSection = new HtmlSection("FPSChart", bStartCollapsed);
+				htmlSection = new HtmlSection("FPSChart", bStartCollapsed, MakeSectionId("FPSChart"));
 				string HeaderRow = "";
 				string ValueRow = "";
 				HeaderRow += "<th>Section Name</th>";
diff --git a/Engine/Source/Programs/CSVTools/PerfReportTool/Summaries/HistogramSummary.cs b/Engine/Source/Programs/CSVTools/PerfReportTool/Summaries/HistogramSummary.cs
--- a/Engine/Source/Programs/CSVTools/PerfReportTool/Summaries/HistogramSummary.cs
+++ b/Engine/Source/Programs/CSVTools/PerfReportTool/Summaries/HistogramSummary.cs
@@ -51,7 +51,7 @@ namespace PerfSummaries
 				return null;
 			}
 
-			HtmlSection htmlSection = new HtmlSection("Histogram", bStartCollapsed);
+			HtmlSection htmlSection = new HtmlSection("Histogram", bStartCollapsed, MakeSectionId("Histogram"));
 
 			// Histogram
 			double[] thresholds = HistogramThresholds;
diff --git a/Engine/Source/Programs/CSVTools/PerfReportTool/Summaries/HitchSummary.cs b/Engine/Source/Programs/CSVTools/PerfReportTool/Summaries/HitchSummary.cs
--- a/Engine/Source/Programs/CSVTools/PerfReportTool/Summaries/HitchSummary.cs
+++ b/Engine/Source/Programs/CSVTools/PerfReportTool/Summaries/HitchSummary.cs
@@ -39,7 +39,7 @@ namespace PerfSummaries
 
 			if (bWriteHtml)
 			{
-				htmlSection = new HtmlSection("Hitches", bStartCollapsed);
+				htmlSection = new HtmlSection("Hitches", bStartCollapsed, MakeSectionId("Hitches"));
 				htmlSection.WriteLine("  <table border='0' style='width:800'>");
 				htmlSection.WriteLine("  <tr><td></td>");
 
diff --git a/Engine/Source/Programs/CSVTools/PerfReportTool/Summaries/MapOverlaySummary.cs b/Engine/Source/Programs/CSVTools/PerfReportTool/Summaries/MapOverlaySummary.cs
--- a/Engine/Source/Programs/CSVTools/PerfReportTool/Summaries/MapOverlaySummary.cs
+++ b/Engine/Source/Programs/CSVTools/PerfReportTool/Summaries/MapOverlaySummary.cs
@@ -290,7 +290,7 @@ namespace PerfSummaries
 			// Output HTML
 			if (bWriteHtml)
 			{
-				htmlSection = new HtmlSection(title, bStartCollapsed);
+				htmlSection = new HtmlSection(title, bStartCollapsed, MakeSectionId(title));
 				string outputDirectory = System.IO.Path.GetDirectoryName(System.IO.Path.GetFullPath(htmlFileName));
 				string outputMapFilename = System.IO.Path.Combine(outputDirectory, destImageFilename);
 
diff --git a/Engine/Source/Programs/CSVTools/PerfReportTool/Summaries/PeakSummary.cs b/Engine/Source/Programs/CSVTools/PerfReportTool/Summaries/PeakSummary.cs
--- a/Engine/Source/Programs/CSVTools/PerfReportTool/Summaries/PeakSummary.cs
+++ b/Engine/Source/Programs/CSVTools/PerfReportTool/Summaries/PeakSummary.cs
@@ -199,7 +199,7 @@ namespace PerfSummaries
 
 			if (bWriteHtml)
 			{
-				htmlSection = new HtmlSection("Peaks Summary", bStartCollapsed);
+				htmlSection = new HtmlSection("Peaks Summary", bStartCollapsed, MakeSectionId("Peaks"));
 			}
 
 			foreach (PeakSummarySection section in peakSummarySections)
diff --git a/Engine/Source/Programs/CSVTools/PerfReportTool/Summaries/StatBudgetsSummary.cs b/Engine/Source/Programs/CSVTools/PerfReportTool/Summaries/StatBudgetsSummary.cs
--- a/Engine/Source/Programs/CSVTools/PerfReportTool/Summaries/StatBudgetsSummary.cs
+++ b/Engine/Source/Programs/CSVTools/PerfReportTool/Summaries/StatBudgetsSummary.cs
@@ -180,7 +180,7 @@ namespace PerfSummaries
 			// Output HTML
 			if (bWriteHtml)
 			{
-				htmlSection = new HtmlSection(title, bStartCollapsed);
+				htmlSection = new HtmlSection(title, bStartCollapsed, MakeSectionId(title));
 				htmlSection.WriteLine("  <table border='0' style='width:400'>");
 				htmlSection.WriteLine("  <tr><th>Stat</th><th>Value (Avg)</th><th>Budget</th></tr>");
 				foreach (StatInfo statInfo in statInfoList)
diff --git a/Engine/Source/Programs/CSVTools/PerfReportTool/Summaries/TimingRegionSummary.cs b/Engine/Source/Programs/CSVTools/PerfReportTool/Summaries/TimingRegionSummary.cs
--- a/Engine/Source/Programs/CSVTools/PerfReportTool/Summaries/TimingRegionSummary.cs
+++ b/Engine/Source/Programs/CSVTools/PerfReportTool/Summaries/TimingRegionSummary.cs
@@ -460,7 +460,7 @@ namespace PerfSummaries
 			string prevCategory = null;
 			if (bWriteHtml)
 			{
-				htmlSection = new HtmlSection(title, bStartCollapsed);
+				htmlSection = new HtmlSection(title, bStartCollapsed, MakeSectionId(title));
 
 				double imageWidth = 1425;
 				double headerWidth = 150;
diff --git a/Engine/Source/Programs/CSVTools/PerfReportTool/Summary.cs b/Engine/Source/Programs/CSVTools/PerfReportTool/Summary.cs
--- a/Engine/Source/Programs/CSVTools/PerfReportTool/Summary.cs
+++ b/Engine/Source/Programs/CSVTools/PerfReportTool/Summary.cs
@@ -1,14 +1,15 @@
 ﻿// Copyright Epic Games, Inc. All Rights Reserved.
 
+using CSVStats;
+using PerfReportTool;
 using System;
 using System.Collections.Generic;
+using System.Drawing;
 using System.Linq;
-using System.Xml.Linq;
-using CSVStats;
-using PerfReportTool;
 using System.Reflection;
-using System.Drawing;
 using System.Text;
+using System.Text.RegularExpressions;
+using System.Xml.Linq;
 
 namespace PerfSummaries
 {
@@ -237,24 +238,37 @@ namespace PerfSummaries
 			return null;
 		}
 
-        public List<CaptureRange> captures;
+		public void SetReportSummaryIndex(int index)
+		{
+			reportSummaryIndex = index;
+		}
+
+		protected string MakeSectionId(string title)
+		{
+			return Regex.Replace(title, @"[^a-zA-Z0-9_-]", "") + "_" + reportSummaryIndex.ToString();
+		}
+
+
+		public List<CaptureRange> captures;
         public List<string> stats;
         public Dictionary<string, ColourThresholdList> StatThresholds;
 		public bool useUnstrippedCsvStats;
 		public bool bStartCollapsed;
 		public bool bHideInDetailedReport = false;
-    };
+		public int reportSummaryIndex = -1;
+	};
 
 
 
 	class HtmlSection
 	{
-		public HtmlSection(string titleIn, bool bStartCollapsedIn, string elementIdIn = null, int headingLevel=2)
+		public HtmlSection(string titleIn, bool bStartCollapsedIn, string elementIdIn = null, int headingLevel=2, int inDisplayWidth=1400)
 		{
 			title = titleIn;
 			bStartCollapsed = bStartCollapsedIn;
 			elementId = elementIdIn;
 			headingType = "h" + headingLevel.ToString();
+			displayWidth = inDisplayWidth;
 		}
 
 		public void WriteLine(string text) 
@@ -305,10 +319,17 @@ namespace PerfSummaries
 				{
 					htmlFile.WriteLine("<a name='" + elementId + "'></a>");
 					extraHeadingAttributes = "id='" + elementId + "'";
+					htmlFile.WriteLine("<div class='sectionAnchorDiv' style='max-width:"+ displayWidth.ToString()+"px'>");
+					htmlFile.WriteLine("<a href='#" + elementId + "' class='sectionAnchorLink'>#</a>");
+					htmlFile.WriteLine("</div>");
 				}
+
 				htmlFile.WriteLine("<"+ headingType + " class='" + headingClass + "' "+ extraHeadingAttributes + ">" + title + "</"+ headingType + ">");
+
 				htmlFile.WriteLine("<div class='" + divClass + "'>");
 				htmlFile.WriteLine("<div class='collapsibleSectionInner'>");
+
+
 			}
 		}
 
@@ -322,13 +343,13 @@ namespace PerfSummaries
 			}
 		}
 
-
 		string title;
 		StringBuilder pendingLine = new StringBuilder();
 		List<string> lines = new List<string>();
 		bool bStartCollapsed;
 		string elementId;
 		string headingType;
+		int displayWidth;
 	};
 
 
