PerfReportTool 4.269.2: HTML section format improvements: only add a left margin for main sections (where heading is h2 or h1). This avoids a double margin for inner sections, and keeps graphs margin-less.

#Synced-CL 46049276

diff --git a/Engine/Source/Programs/CSVTools/PerfReportTool/PerfReportTool.cs b/Engine/Source/Programs/CSVTools/PerfReportTool/PerfReportTool.cs
--- a/Engine/Source/Programs/CSVTools/PerfReportTool/PerfReportTool.cs
+++ b/Engine/Source/Programs/CSVTools/PerfReportTool/PerfReportTool.cs
@@ -24,7 +24,7 @@ namespace PerfReportTool
 	class Version
 	{
 		// Format: Major.Minor.Bugfix
-		private static string VersionString = "4.269.1";
+		private static string VersionString = "4.269.2";
 
 		public static string Get() { return VersionString; }
 	};
@@ -1619,12 +1619,12 @@ namespace PerfReportTool
 				htmlFile.WriteLine("      .collapsibleHeading:hover { background-color: #c4c4c4; }");
 				htmlFile.WriteLine("      .collapsibleHeading:after { content: '\u25B2'; font-size: 14px; color: #4d4d4d; font-variant: normal; float: right; margin-left: 14px;}");
 				htmlFile.WriteLine("      .collapsibleHeading.expanded:after { content: '\u25BC'; }");
-				htmlFile.WriteLine("      .collapsibleSection { display: grid; grid-template-rows: 0fr; transition: grid-template-rows 0.2s ease-out; margin-left:5px }");
+				htmlFile.WriteLine("      .collapsibleSection { display: grid; grid-template-rows: 0fr; transition: grid-template-rows 0.2s ease-out }");
 				htmlFile.WriteLine("      .collapsibleSection.expanded { grid-template-rows: 1fr; }");
-				htmlFile.WriteLine("      .collapsibleSectionInner { overflow: auto; overflow-y: hidden; }");
+				htmlFile.WriteLine("      .collapsibleSectionInner { overflow: auto; overflow-y: hidden }");
 
 				// Section anchors
-				htmlFile.WriteLine("      .sectionAnchorLink {font-size:16px; color:#808080; text-decoration:none;display: inline-block;vertical-align:-4px}");
+				htmlFile.WriteLine("      .sectionAnchorLink { color:#808080; text-decoration:none;display: inline-block; vertical-align:-4px; margin-right:-2px; margin-left:-2px }");
 				htmlFile.WriteLine("      .sectionAnchorText-h1 {font-size:16px}");
 				htmlFile.WriteLine("      .sectionAnchorText-h2 {font-size:14px}");
 				htmlFile.WriteLine("      .sectionAnchorText-h3 {font-size:12px}");
@@ -1824,7 +1824,7 @@ namespace PerfReportTool
 					htmlSection.WriteToFile(htmlFile);
 				}
 
-				htmlFile.WriteLine("<a href='#top'>Back to top \u2191</a>");
+				htmlFile.WriteLine("<br><a href='#top'>Back to top \u2191</a>");
 
 
 				if (GetBoolArg("noWatermarks"))
diff --git a/Engine/Source/Programs/CSVTools/PerfReportTool/Summary.cs b/Engine/Source/Programs/CSVTools/PerfReportTool/Summary.cs
--- a/Engine/Source/Programs/CSVTools/PerfReportTool/Summary.cs
+++ b/Engine/Source/Programs/CSVTools/PerfReportTool/Summary.cs
@@ -262,11 +262,12 @@ namespace PerfSummaries
 
 	class HtmlSection
 	{
-		public HtmlSection(string titleIn, bool bStartCollapsedIn, string elementIdIn = null, int headingLevel=2)
+		public HtmlSection(string titleIn, bool bStartCollapsedIn, string elementIdIn = null, int headingLevelIn=2)
 		{
 			title = titleIn;
 			bStartCollapsed = bStartCollapsedIn;
 			elementId = elementIdIn;
+			headingLevel = headingLevelIn;
 			headingType = "h" + headingLevel.ToString();
 		}
 
@@ -328,7 +329,14 @@ namespace PerfSummaries
 				htmlFile.WriteLine("<"+ headingType + " class='" + headingClass + "' "+ extraHeadingAttributes + ">" + title + "</"+ headingType + ">");
 
 				htmlFile.WriteLine("<div class='" + divClass + "'>");
-				htmlFile.WriteLine("<div class='collapsibleSectionInner'>");
+
+				string extraSectionAttributes = "";
+				// If this is a main section then add padding
+				if (headingLevel <= 2)
+				{
+					extraSectionAttributes += "style='margin-left:12px'";
+				}
+				htmlFile.WriteLine("<div class='collapsibleSectionInner' "+ extraSectionAttributes + ">");
 			}
 		}
 
@@ -346,6 +354,7 @@ namespace PerfSummaries
 		StringBuilder pendingLine = new StringBuilder();
 		List<string> lines = new List<string>();
 		bool bStartCollapsed;
+		int headingLevel;
 		string elementId;
 		string headingType;
 	};
