Add a P4SyncCsvTools batch file. This should be run before publishing

diff --git a/Engine/Source/Programs/CSVTools/P4SyncCsvTools.bat b/Engine/Source/Programs/CSVTools/P4SyncCsvTools.bat
new file mode 100755
--- /dev/null
+++ b/Engine/Source/Programs/CSVTools/P4SyncCsvTools.bat
@@ -0,0 +1,14 @@
+@echo off
+setlocal enabledelayedexpansion
+
+pushd "%~dp0\..\..\..
+
+for /f "tokens=1" %%A in ('p4 where . ^| findstr /c:"//"') do (
+    set p4Path=%%A
+)
+echo P4 engine Path: !p4Path!
+
+p4 sync !p4Path!/Binaries/DotNET/CsvTools/...
+p4 sync !p4Path!/Source/Programs/CSVTools/...
+p4 sync !p4Path!/Source/Programs/Shared/...
+popd
\ No newline at end of file
