[TimeToMatch] [Part 1/3 (Engine Changes)] Add new timing regions property (findLast) to be able to pick the last occurrence of an event pair instead of first (as begin/end event/region).

Required for tracking matchmaking times as players can cancel/resume matchmaking, thus generating multiple regions, but we only care about the last (the successful attempt) in the CSV.

#tests
#rnx

#Synced-CL 51169632

diff --git a/Engine/Source/Programs/CSVTools/CsvStats/CsvStats.cs b/Engine/Source/Programs/CSVTools/CsvStats/CsvStats.cs
--- a/Engine/Source/Programs/CSVTools/CsvStats/CsvStats.cs
+++ b/Engine/Source/Programs/CSVTools/CsvStats/CsvStats.cs
@@ -529,9 +529,24 @@ namespace CSVStats
 			return metaData.GetValue(key, null);
 		}
 
+		public Tuple<CsvEvent, CsvEvent> FindEventPair(string patternBegin, string patternEnd, bool bFindLast = false)
+		{
+			IEnumerable<CsvEvent> eventsOrdered = bFindLast == false ? Events : Events.Reverse<CsvEvent>();
+			CsvEvent beginEvent = patternBegin != null? FindEventInternal(eventsOrdered, patternBegin) : null;
+			double startTime = (beginEvent != null) ? beginEvent.Timestamp : 0.0;
+
+			CsvEvent endEvent = (patternEnd != null) ? FindEventInternal(eventsOrdered, patternEnd, startTime, double.MaxValue) : null;
+			return Tuple.Create(beginEvent, endEvent);
+		}
+
 		public CsvEvent FindEvent(string pattern, double startTime = 0.0, double endTime = double.MaxValue)
 		{
-			foreach (CsvEvent ev in Events)
+			return FindEventInternal(Events, pattern, startTime, endTime);
+		}
+
+		private CsvEvent FindEventInternal(IEnumerable<CsvEvent> InEvents, string pattern, double startTime = 0.0, double endTime = double.MaxValue)
+		{
+			foreach (CsvEvent ev in InEvents)
 			{
 				if (CsvStats.DoesSearchStringMatch(ev.Name, pattern) && ev.Timestamp >= startTime && ev.Timestamp <= endTime)
 				{
@@ -542,9 +557,14 @@ namespace CSVStats
 		}
 
 		public List<CsvEvent> FindEvents(string pattern, double startTime = 0.0, double endTime = double.MaxValue)
+		{
+			return FindEventsInternal(Events, pattern, startTime, endTime);
+		}
+
+		private List<CsvEvent> FindEventsInternal(IEnumerable<CsvEvent> InEvents, string pattern, double startTime = 0.0, double endTime = double.MaxValue)
 		{
 			List<CsvEvent> eventList = new List<CsvEvent>();
-			foreach (CsvEvent ev in Events)
+			foreach (CsvEvent ev in InEvents)
 			{
 				if (CsvStats.DoesSearchStringMatch(ev.Name, pattern) && ev.Timestamp >= startTime && ev.Timestamp <= endTime)
 				{
diff --git a/Engine/Source/Programs/CSVTools/PerfReportTool/PerfReportTool.cs b/Engine/Source/Programs/CSVTools/PerfReportTool/PerfReportTool.cs
--- a/Engine/Source/Programs/CSVTools/PerfReportTool/PerfReportTool.cs
+++ b/Engine/Source/Programs/CSVTools/PerfReportTool/PerfReportTool.cs
@@ -24,7 +24,7 @@ namespace PerfReportTool
 	class Version
 	{
 		// Format: Major.Minor.Bugfix
-		private static string VersionString = "4.269.2";
+		private static string VersionString = "4.269.3";
 
 		public static string Get() { return VersionString; }
 	};
diff --git a/Engine/Source/Programs/CSVTools/PerfReportTool/Summaries/TimingRegionSummary.cs b/Engine/Source/Programs/CSVTools/PerfReportTool/Summaries/TimingRegionSummary.cs
--- a/Engine/Source/Programs/CSVTools/PerfReportTool/Summaries/TimingRegionSummary.cs
+++ b/Engine/Source/Programs/CSVTools/PerfReportTool/Summaries/TimingRegionSummary.cs
@@ -59,7 +59,8 @@ namespace PerfSummaries
 			}
 
 			bShowTimingTable = element.GetSafeAttribute<bool>(vars, "showTimingTable", true);
-			
+			bFindLast = element.GetSafeAttribute<bool>(vars, "findLast", false);
+
 			metricPrefix = element.GetSafeAttribute<string>(vars, "metricPrefix");
 			foreach (XElement regionFilterXml in element.Elements("regionFilter"))
 			{
@@ -441,10 +442,12 @@ namespace PerfSummaries
 				throw new Exception("TimingRegion summary requires a frametime stat");
 			}
 
-			CsvEvent beginEvent = (beginEventName != null) ? csvStatsUnstripped.FindEvent(beginEventName) : null;
+			Tuple<CsvEvent, CsvEvent> eventPair = csvStatsUnstripped.FindEventPair(beginEventName, endEventName, bFindLast);
+
+			CsvEvent beginEvent = eventPair.Item1;
 			double startTime = (beginEvent != null) ? beginEvent.Timestamp : 0.0;
 
-			CsvEvent endEvent = (endEventName != null) ? csvStatsUnstripped.FindEvent(endEventName, startTime) : null;
+			CsvEvent endEvent = eventPair.Item2;
 			double endTime = (endEvent != null) ? endEvent.Timestamp : frameTimeStat.total / 1000.0;
 
 			double totalDuration = endTime - startTime;
@@ -586,6 +589,7 @@ namespace PerfSummaries
 		string endEventName;
 		string metricPrefix;
 		bool bShowTimingTable;
+		bool bFindLast;
 		List<RegionFilterInfo> regionFilters = new List<RegionFilterInfo>();
 	};
 
