PRT - Add SlowFrameTimePercent to FPSChartSummary.

#rnx

#Synced-CL 51247236

#Synced-CL 51247236

diff --git a/Engine/Source/Programs/CSVTools/PerfReportTool/PerfReportTool.cs b/Engine/Source/Programs/CSVTools/PerfReportTool/PerfReportTool.cs
--- a/Engine/Source/Programs/CSVTools/PerfReportTool/PerfReportTool.cs
+++ b/Engine/Source/Programs/CSVTools/PerfReportTool/PerfReportTool.cs
@@ -24,7 +24,7 @@ namespace PerfReportTool
 	class Version
 	{
 		// Format: Major.Minor.Bugfix
-		private static string VersionString = "4.269.3";
+		private static string VersionString = "4.270.0";
 
 		public static string Get() { return VersionString; }
 	};
diff --git a/Engine/Source/Programs/CSVTools/PerfReportTool/Summaries/FPSChartSummary.cs b/Engine/Source/Programs/CSVTools/PerfReportTool/Summaries/FPSChartSummary.cs
--- a/Engine/Source/Programs/CSVTools/PerfReportTool/Summaries/FPSChartSummary.cs
+++ b/Engine/Source/Programs/CSVTools/PerfReportTool/Summaries/FPSChartSummary.cs
@@ -16,6 +16,8 @@ namespace PerfSummaries
 		{
 			ReadStatsFromXML(element, vars);
 			fps = element.GetRequiredAttribute<int>(vars, "fps");
+			recommendedFps = element.GetSafeAttribute<int>(vars, "recommendedFps", 0);
+			slowFrameTimeBuffer = element.GetSafeAttribute<float>(vars, "slowFrameTimeBuffer", 0);
 			hitchThreshold = (float)element.GetRequiredAttribute<double>(vars, "hitchThreshold");
 			bUseEngineHitchMetric = element.GetSafeAttribute<bool>(vars, "useEngineHitchMetric", false);
 			if (bUseEngineHitchMetric)
@@ -26,6 +28,7 @@ namespace PerfSummaries
 
 			bIgnoreHitchTimePercent = element.GetSafeAttribute<bool>(vars, "ignoreHitchTimePercent", false);
 			bIgnoreMVP = element.GetSafeAttribute<bool>(vars, "ignoreMVP", false);
+			bIgnoreSlowFrameTimePercent = element.GetSafeAttribute<bool>(vars, "ignoreSlowFrameTimePercent", false);
 		}
 
 		public FPSChartSummary() { }
@@ -48,6 +51,7 @@ namespace PerfSummaries
 			public float HitchTimePercent;
 			public int HitchCount;
 			public float TotalTimeSeconds;
+			public float? SlowFrameTimePerecent;
 		};
 
 		FpsChartData ComputeFPSChartDataForFrames(List<float> frameTimes, bool skiplastFrame)
@@ -117,6 +121,32 @@ namespace PerfSummaries
 			totalHitchTime -= (hitchCount * hitchThreshold);
 			outData.HitchTimePercent = (float)(totalHitchTime / totalFrametime) * 100.0f;
 
+			if (!bIgnoreSlowFrameTimePercent && recommendedFps > 0 && TotalSeconds > 0.0f)
+			{
+				float frameTimeTargetMs = (1.0f / (float)recommendedFps) * 1000.0f;
+				// We allow for some tolerance in determining what is a slow frame to account for natural noise in timings.
+				float frameTimeThreshold = frameTimeTargetMs + slowFrameTimeBuffer;
+
+				// Accumulate the total time spent over the frame time target each frame.
+				float totalOverBudget = 0.0f;
+				for (int i = 0; i < frameCount; i++)
+				{
+					if (frameTimes[i] > frameTimeThreshold)
+					{
+						totalOverBudget += frameTimes[i] - frameTimeTargetMs;
+					}
+				}
+
+				// The final value is the % of time spent over the frame time target.
+				float totalOverBudgetSeconds = totalOverBudget / 1000.0f;
+				outData.SlowFrameTimePerecent = (totalOverBudgetSeconds / TotalSeconds) * 100.0f; // Scale 0-100
+			}
+			else
+			{
+				// Avoid outputting 0 if the required values weren't available as it may be misleading.
+				outData.SlowFrameTimePerecent = null;
+			}
+
 			// If frame count is low enough this can get rounded down to 0
 			int TotalTargetFrames = (int)((double)fps * (TotalSeconds));
 			if (TotalTargetFrames > 0)
@@ -202,6 +232,11 @@ namespace PerfSummaries
 				Columns.Add(new ColumnInfo("HitchTimePercent", fpsChartData.HitchTimePercent, GetStatColourThresholdList("HitchTimePercent")));
 			}
 
+			if (!bIgnoreSlowFrameTimePercent && fpsChartData.SlowFrameTimePerecent.HasValue)
+			{
+				Columns.Add(new ColumnInfo("SlowFrameTimePercent", fpsChartData.SlowFrameTimePerecent.Value, GetStatColourThresholdList("SlowFrameTimePercent")));
+			}
+
 			string MvpStatName = "MVP" + fps.ToString();
 			if (!bIgnoreMVP)
 			{
@@ -390,6 +425,11 @@ namespace PerfSummaries
 						ColumnDict[MvpStatName].Value = captureFpsChartData.MVP;
 					}
 
+					if (!bIgnoreSlowFrameTimePercent && captureFpsChartData.SlowFrameTimePerecent.HasValue)
+					{
+						ColumnDict["SlowFrameTimePercent"].Value = captureFpsChartData.SlowFrameTimePerecent.Value;
+					}
+
 					// Update the CSV stat values
 					int columnIndex = csvStatColumnStartIndex;
 					foreach (string statName in stats)
@@ -438,6 +478,11 @@ namespace PerfSummaries
 							rowData.Add(SummaryTableElement.Type.SummaryTableMetric, CapRange.name + "_MVP", captureFpsChartData.MVP, GetStatColourThresholdList(MvpStatName));
 						}
 
+						if (!bIgnoreSlowFrameTimePercent && captureFpsChartData.SlowFrameTimePerecent.HasValue)
+						{
+							rowData.Add(SummaryTableElement.Type.SummaryTableMetric, CapRange.name + "_SlowFrameTimePercent", captureFpsChartData.SlowFrameTimePerecent.Value, GetStatColourThresholdList("SlowFrameTimePerecent"));
+						}
+
 						rowData.Add(SummaryTableElement.Type.SummaryTableMetric, CapRange.name + "_HPM", captureFpsChartData.HitchesPerMinute, GetStatColourThresholdList("Hitches/Min"));
 						rowData.Add(SummaryTableElement.Type.SummaryTableMetric, CapRange.name + "_HitchCount", captureFpsChartData.HitchCount);
 						rowData.Add(SummaryTableElement.Type.SummaryTableMetric, CapRange.name + "_TotalTimeSeconds", captureFpsChartData.TotalTimeSeconds);
@@ -493,9 +538,12 @@ namespace PerfSummaries
 		}
 
 		int fps;
+		int recommendedFps;
+		float slowFrameTimeBuffer;
 		float hitchThreshold;
 		bool bUseEngineHitchMetric;
 		bool bIgnoreHitchTimePercent;
+		bool bIgnoreSlowFrameTimePercent;
 		bool bIgnoreMVP;
 		float engineHitchToNonHitchRatio;
 		float engineMinTimeBetweenHitchesMs;
