PRT - Fix issue in SlowFrameTime calculation where it wouldn't use TargetFps if that is lower than recommended.

#rnx
#Synced-CL 51247236

diff --git a/Engine/Source/Programs/CSVTools/PerfReportTool/PerfReportTool.cs b/Engine/Source/Programs/CSVTools/PerfReportTool/PerfReportTool.cs
--- a/Engine/Source/Programs/CSVTools/PerfReportTool/PerfReportTool.cs
+++ b/Engine/Source/Programs/CSVTools/PerfReportTool/PerfReportTool.cs
@@ -24,7 +24,7 @@ namespace PerfReportTool
 	class Version
 	{
 		// Format: Major.Minor.Bugfix
-		private static string VersionString = "4.270.0";
+		private static string VersionString = "4.270.1";
 
 		public static string Get() { return VersionString; }
 	};
diff --git a/Engine/Source/Programs/CSVTools/PerfReportTool/Summaries/FPSChartSummary.cs b/Engine/Source/Programs/CSVTools/PerfReportTool/Summaries/FPSChartSummary.cs
--- a/Engine/Source/Programs/CSVTools/PerfReportTool/Summaries/FPSChartSummary.cs
+++ b/Engine/Source/Programs/CSVTools/PerfReportTool/Summaries/FPSChartSummary.cs
@@ -121,9 +121,12 @@ namespace PerfSummaries
 			totalHitchTime -= (hitchCount * hitchThreshold);
 			outData.HitchTimePercent = (float)(totalHitchTime / totalFrametime) * 100.0f;
 
-			if (!bIgnoreSlowFrameTimePercent && recommendedFps > 0 && TotalSeconds > 0.0f)
+			if (!bIgnoreSlowFrameTimePercent && recommendedFps > 0 && fps > 0 && TotalSeconds > 0.0f)
 			{
-				float frameTimeTargetMs = (1.0f / (float)recommendedFps) * 1000.0f;
+				// If the user is targeting a lower fps than we recommend then use that as the baseline instead.
+				// Otherwise we'll incorrectly compute a high SFP due to the frame rate being limited to a lower value.
+				float minTargetFps = (float)Math.Min(recommendedFps, fps);
+				float frameTimeTargetMs = (1.0f / minTargetFps) * 1000.0f;
 				// We allow for some tolerance in determining what is a slow frame to account for natural noise in timings.
 				float frameTimeThreshold = frameTimeTargetMs + slowFrameTimeBuffer;
 
