PRT - Fix SummaryTableColumn.AreAllValuesOverThreshold to only return true if all values are over the threshold.

#rnx
#Synced-CL 52552752

diff --git a/Engine/Source/Programs/CSVTools/PerfReportTool/PerfReportTool.cs b/Engine/Source/Programs/CSVTools/PerfReportTool/PerfReportTool.cs
--- a/Engine/Source/Programs/CSVTools/PerfReportTool/PerfReportTool.cs
+++ b/Engine/Source/Programs/CSVTools/PerfReportTool/PerfReportTool.cs
@@ -24,7 +24,7 @@ namespace PerfReportTool
 	class Version
 	{
 		// Format: Major.Minor.Bugfix
-		private static string VersionString = "4.272.0";
+		private static string VersionString = "4.272.1";
 
 		public static string Get() { return VersionString; }
 	};
diff --git a/Engine/Source/Programs/CSVTools/PerfReportTool/SummaryTable.cs b/Engine/Source/Programs/CSVTools/PerfReportTool/SummaryTable.cs
--- a/Engine/Source/Programs/CSVTools/PerfReportTool/SummaryTable.cs
+++ b/Engine/Source/Programs/CSVTools/PerfReportTool/SummaryTable.cs
@@ -1079,14 +1079,19 @@ namespace PerfSummaries
 			{
 				return true;
 			}
+			bool hasRealValue = false;
 			foreach(double value in doubleValues)
 			{
-				if ( value > threshold && value != double.MaxValue)
+				if (value != double.MaxValue)
 				{
-					return true;
+					if (value <= threshold)
+					{
+						return false;
+					}
+					hasRealValue = true;
 				}
 			}
-			return false;
+			return hasRealValue;
 		}
 
 		public void SetStringValue(int index, string value)
