PRT - Add unit and integration tests for PerfReportTool.

+ Uses csvs from citysample as test data.
+ Covers most core functionality of PRT.

#rnx
#Synced-CL 52552752

diff --git a/Engine/Source/Programs/CSVTools/PerfReportTool.Tests/CsvStatsTests/CsvMetadataTests.cs b/Engine/Source/Programs/CSVTools/PerfReportTool.Tests/CsvStatsTests/CsvMetadataTests.cs
new file mode 100644
--- /dev/null
+++ b/Engine/Source/Programs/CSVTools/PerfReportTool.Tests/CsvStatsTests/CsvMetadataTests.cs
@@ -0,0 +1,139 @@
+using CSVStats;
+using PerfReportTool.Tests.TestHelpers;
+using Xunit;
+
+namespace PerfReportTool.Tests.CsvStatsTests
+{
+	public class CsvMetadataTests
+	{
+		[Fact]
+		public void Constructor_Default_CreatesEmptyValues()
+		{
+			var metadata = new CsvMetadata();
+			Assert.NotNull(metadata.Values);
+			Assert.Empty(metadata.Values);
+		}
+
+		[Fact]
+		public void ParseFromCsvLine_ParsesKeyValuePairs()
+		{
+			// CSV metadata format: [key],value,[key],value,...
+			string[] line = new[] { "[platform]", "windows", "[buildversion]", "5.8.0" };
+			var metadata = new CsvMetadata(line);
+
+			Assert.Equal("windows", metadata.Values["platform"]);
+			Assert.Equal("5.8.0", metadata.Values["buildversion"]);
+		}
+
+		[Fact]
+		public void ParseFromCsvLine_KeysAreLowercase()
+		{
+			string[] line = new[] { "[Platform]", "Windows", "[BuildVersion]", "5.8.0" };
+			var metadata = new CsvMetadata(line);
+
+			// Keys should be lowercased, values preserved
+			Assert.True(metadata.Values.ContainsKey("platform"));
+			Assert.Equal("Windows", metadata.Values["platform"]);
+		}
+
+		[Fact]
+		public void ParseFromCsvLine_CommandlineSpecialCasing_HandlesCommas()
+		{
+			// The commandline key gets special treatment - everything after it is joined with commas
+			string[] line = new[] { "[platform]", "windows", "[commandline]", "\"-arg1", "val1", "-arg2", "val2\"" };
+			var metadata = new CsvMetadata(line);
+
+			Assert.True(metadata.Values.ContainsKey("commandline"));
+			string cmdline = metadata.Values["commandline"];
+			// All parts should be reassembled with commas, quotes stripped from first/last
+			Assert.Contains("-arg1", cmdline);
+			Assert.Contains("val1", cmdline);
+			Assert.Contains("-arg2", cmdline);
+			Assert.Contains("val2", cmdline);
+		}
+
+		[Fact]
+		public void CopyConstructor_ProducesIndependentCopy()
+		{
+			var original = CsvTestHelper.CreateMetadata(("platform", "windows"), ("config", "test"));
+			var copy = new CsvMetadata(original);
+
+			// Should have the same values
+			Assert.Equal(original.Values["platform"], copy.Values["platform"]);
+
+			// Modifying copy should not affect original
+			copy.Values["platform"] = "linux";
+			Assert.Equal("windows", original.Values["platform"]);
+			Assert.Equal("linux", copy.Values["platform"]);
+		}
+
+		[Fact]
+		public void Clone_ProducesCopy()
+		{
+			var original = CsvTestHelper.CreateMetadata(("platform", "windows"));
+			var clone = original.Clone();
+
+			Assert.Equal("windows", clone.Values["platform"]);
+		}
+
+		[Fact]
+		public void Matches_IdenticalMetadata_ReturnsTrue()
+		{
+			var a = CsvTestHelper.CreateMetadata(("platform", "windows"), ("config", "test"));
+			var b = CsvTestHelper.CreateMetadata(("platform", "windows"), ("config", "test"));
+			Assert.True(CsvMetadata.Matches(a, b));
+		}
+
+		[Fact]
+		public void Matches_DifferentValues_ReturnsFalse()
+		{
+			var a = CsvTestHelper.CreateMetadata(("platform", "windows"));
+			var b = CsvTestHelper.CreateMetadata(("platform", "linux"));
+			Assert.False(CsvMetadata.Matches(a, b));
+		}
+
+		[Fact]
+		public void Matches_DifferentKeys_ReturnsFalse()
+		{
+			var a = CsvTestHelper.CreateMetadata(("platform", "windows"));
+			var b = CsvTestHelper.CreateMetadata(("config", "test"));
+			Assert.False(CsvMetadata.Matches(a, b));
+		}
+
+		[Fact]
+		public void Matches_DifferentCounts_ReturnsFalse()
+		{
+			var a = CsvTestHelper.CreateMetadata(("platform", "windows"));
+			var b = CsvTestHelper.CreateMetadata(("platform", "windows"), ("config", "test"));
+			Assert.False(CsvMetadata.Matches(a, b));
+		}
+
+		[Fact]
+		public void Matches_BothNull_ReturnsTrue()
+		{
+			Assert.True(CsvMetadata.Matches(null, null));
+		}
+
+		[Fact]
+		public void Matches_OneNull_ReturnsFalse()
+		{
+			var a = CsvTestHelper.CreateMetadata(("platform", "windows"));
+			Assert.False(CsvMetadata.Matches(a, null));
+			Assert.False(CsvMetadata.Matches(null, a));
+		}
+
+		[Fact]
+		public void GetValue_ExistingKey_ReturnsValue()
+		{
+			var metadata = CsvTestHelper.CreateMetadata(("platform", "windows"));
+			Assert.Equal("windows", metadata.GetValue("platform", "default"));
+		}
+
+		[Fact]
+		public void GetValue_MissingKey_ReturnsDefault()
+		{
+			var metadata = CsvTestHelper.CreateMetadata(("platform", "windows"));
+			Assert.Equal("default", metadata.GetValue("missing", "default"));
+		}
+	}
+}
diff --git a/Engine/Source/Programs/CSVTools/PerfReportTool.Tests/CsvStatsTests/CsvStatsReadTests.cs b/Engine/Source/Programs/CSVTools/PerfReportTool.Tests/CsvStatsTests/CsvStatsReadTests.cs
new file mode 100644
--- /dev/null
+++ b/Engine/Source/Programs/CSVTools/PerfReportTool.Tests/CsvStatsTests/CsvStatsReadTests.cs
@@ -0,0 +1,186 @@
+using System.IO;
+using CSVStats;
+using Xunit;
+
+namespace PerfReportTool.Tests.CsvStatsTests
+{
+	[Trait("Category", "FileIO")]
+	public class CsvStatsReadTests : IClassFixture<TestDataFixture>
+	{
+		private readonly TestDataFixture _fixture;
+
+		// File name constants for the in-project test data
+		private const string PerfCsv1 = "Perf_++UE5+Main-CL-52441555-CitySample_Windows_B02C1ABC494618E056385CB0080C1931.csv.bin";
+		private const string PerfCsv2 = "Perf_++UE5+Main-CL-52441555-CitySample_Windows_F11D38AE4A4086B50D8FB7A37E337B3E.csv.bin";
+		private const string LlmCsv1 = "LLM_++UE5+Main-CL-52441555-CitySample_Windows_126E42324728AE04A37A70B29E582593.csv.bin";
+		private const string LlmCsv2 = "LLM_++UE5+Main-CL-52441555-CitySample_Windows_2C2C8B3D49A805590E22DD9E0AF2FEF5.csv.bin";
+		private const string NoMetadataCsv = "perf_csv_nometadata.bin";
+
+		public CsvStatsReadTests(TestDataFixture fixture)
+		{
+			_fixture = fixture;
+		}
+
+		[Fact]
+		public void ReadCSVFile_BinaryCsvBin_ReadsStatsAndMetadata()
+		{
+			_fixture.RequireTestData();
+
+			string csvPath = _fixture.GetFile(PerfCsv1);
+			Assert.True(File.Exists(csvPath), $"Test file not found: {csvPath}");
+
+			var stats = CsvStats.ReadCSVFile(csvPath);
+			Assert.NotNull(stats);
+			Assert.True(stats.Stats.Count > 0, "Expected at least one stat from binary CSV");
+			Assert.NotNull(stats.metaData);
+		}
+
+		[Fact]
+		public void ReadCSVFile_NoMetadataCsv_HandlesMissingMetadata()
+		{
+			_fixture.RequireTestData();
+
+			string csvPath = _fixture.GetFile(NoMetadataCsv);
+			Assert.True(File.Exists(csvPath), $"No-metadata test file not found: {csvPath}");
+
+			var stats = CsvStats.ReadCSVFile(csvPath);
+			Assert.NotNull(stats);
+			Assert.True(stats.Stats.Count > 0, "Expected stats even without metadata");
+			Assert.Null(stats.metaData);
+		}
+
+		[Fact]
+		public void ReadCSVFile_LlmCsv_ReadsMemoryStats()
+		{
+			_fixture.RequireTestData();
+
+			string csvPath = _fixture.GetFile(LlmCsv1);
+			Assert.True(File.Exists(csvPath), $"LLM CSV not found: {csvPath}");
+
+			var stats = CsvStats.ReadCSVFile(csvPath);
+			Assert.NotNull(stats);
+			Assert.True(stats.Stats.Count > 0, "Expected memory stats from LLM CSV");
+			Assert.NotNull(stats.metaData);
+			Assert.Equal("1", stats.metaData.GetValue("llm", "0"));
+		}
+
+		[Fact]
+		public void ReadCSVFile_BinaryCsvBin_HasExpectedStatTypes()
+		{
+			_fixture.RequireTestData();
+
+			string csvPath = _fixture.GetFile(PerfCsv1);
+			Assert.True(File.Exists(csvPath), $"Binary CSV not found: {csvPath}");
+
+			var stats = CsvStats.ReadCSVFile(csvPath);
+			Assert.NotNull(stats);
+
+			// Perf CSVs should have many stats (typically hundreds)
+			Assert.True(stats.Stats.Count > 10, $"Expected many stats but got {stats.Stats.Count}");
+
+			// Verify samples are non-empty
+			bool hasNonEmptyStat = false;
+			foreach (var stat in stats.Stats.Values)
+			{
+				if (stat.GetNumSamples() > 0)
+				{
+					hasNonEmptyStat = true;
+					break;
+				}
+			}
+			Assert.True(hasNonEmptyStat, "Expected at least one stat with samples");
+		}
+
+		[Fact]
+		public void ReadCSVFile_JustHeader_ReadsOnlyStatNames()
+		{
+			_fixture.RequireTestData();
+
+			string csvPath = _fixture.GetFile(PerfCsv1);
+			Assert.True(File.Exists(csvPath), $"Test file not found: {csvPath}");
+
+			var stats = CsvStats.ReadCSVFile(csvPath, bJustHeader: true);
+			Assert.NotNull(stats);
+
+			// Should have stat names but empty samples
+			Assert.True(stats.Stats.Count > 0, "Expected stat definitions from header");
+			foreach (var stat in stats.Stats.Values)
+			{
+				Assert.Equal(0, stat.GetNumSamples());
+			}
+		}
+
+		[Fact]
+		public void ReadCSVFile_WithStatFilter_OnlyReturnsMatchingStats()
+		{
+			_fixture.RequireTestData();
+
+			string csvPath = _fixture.GetFile(PerfCsv1);
+			Assert.True(File.Exists(csvPath), $"Test file not found: {csvPath}");
+
+			// Read all stats first to know what exists
+			var allStats = CsvStats.ReadCSVFile(csvPath);
+			Assert.True(allStats.Stats.Count > 1, "Need more than 1 stat for this test");
+
+			// Read with a filter for just one stat
+			string firstStatName = null;
+			foreach (var key in allStats.Stats.Keys)
+			{
+				firstStatName = key;
+				break;
+			}
+
+			var filteredStats = CsvStats.ReadCSVFile(csvPath, statNames: new[] { firstStatName });
+			Assert.NotNull(filteredStats);
+			Assert.True(filteredStats.Stats.Count <= allStats.Stats.Count);
+		}
+
+		[Fact]
+		public void ReadCSVFile_MetadataContainsExpectedKeys()
+		{
+			_fixture.RequireTestData();
+
+			string csvPath = _fixture.GetFile(PerfCsv1);
+			Assert.True(File.Exists(csvPath), $"Test file not found: {csvPath}");
+
+			var stats = CsvStats.ReadCSVFile(csvPath);
+			Assert.NotNull(stats.metaData);
+
+			// UE5 perf CSV metadata should contain standard keys
+			Assert.True(stats.metaData.Values.Count > 0, "Expected metadata key-value pairs");
+			Assert.Equal("Windows", stats.metaData.GetValue("platform", null));
+			Assert.NotNull(stats.metaData.GetValue("buildversion", null));
+		}
+
+		[Fact]
+		public void ReadCSVFile_TwoPerfCsvs_HaveSimilarStatCounts()
+		{
+			_fixture.RequireTestData();
+
+			string csvPath1 = _fixture.GetFile(PerfCsv1);
+			string csvPath2 = _fixture.GetFile(PerfCsv2);
+
+			var stats1 = CsvStats.ReadCSVFile(csvPath1, bJustHeader: true);
+			var stats2 = CsvStats.ReadCSVFile(csvPath2, bJustHeader: true);
+
+			// Both perf CSVs from the same build should have roughly the same stat count
+			// (minor differences are possible due to dynamic stats)
+			double ratio = (double)stats1.Stats.Count / stats2.Stats.Count;
+			Assert.True(ratio > 0.95 && ratio < 1.05,
+				$"Stat counts differ significantly: {stats1.Stats.Count} vs {stats2.Stats.Count}");
+		}
+
+		[Fact]
+		public void ReadCSVFile_LlmCsv_HasMoreStatsThanPerfCsv()
+		{
+			_fixture.RequireTestData();
+
+			var perfStats = CsvStats.ReadCSVFile(_fixture.GetFile(PerfCsv1), bJustHeader: true);
+			var llmStats = CsvStats.ReadCSVFile(_fixture.GetFile(LlmCsv1), bJustHeader: true);
+
+			// LLM CSVs include additional LLM/* memory stats
+			Assert.True(llmStats.Stats.Count > perfStats.Stats.Count,
+				$"Expected LLM ({llmStats.Stats.Count} stats) to have more stats than Perf ({perfStats.Stats.Count} stats)");
+		}
+	}
+}
diff --git a/Engine/Source/Programs/CSVTools/PerfReportTool.Tests/CsvStatsTests/CsvStatsUtilityTests.cs b/Engine/Source/Programs/CSVTools/PerfReportTool.Tests/CsvStatsTests/CsvStatsUtilityTests.cs
new file mode 100644
--- /dev/null
+++ b/Engine/Source/Programs/CSVTools/PerfReportTool.Tests/CsvStatsTests/CsvStatsUtilityTests.cs
@@ -0,0 +1,183 @@
+using CSVStats;
+using PerfReportTool.Tests.TestHelpers;
+using Xunit;
+
+namespace PerfReportTool.Tests.CsvStatsTests
+{
+	public class DoesSearchStringMatchTests
+	{
+		[Fact]
+		public void ExactMatch_ReturnsTrue()
+		{
+			Assert.True(CSVStats.CsvStats.DoesSearchStringMatch("frametime", "frametime"));
+		}
+
+		[Fact]
+		public void CaseInsensitive_ReturnsTrue()
+		{
+			Assert.True(CSVStats.CsvStats.DoesSearchStringMatch("FrameTime", "frametime"));
+			Assert.True(CSVStats.CsvStats.DoesSearchStringMatch("frametime", "FrameTime"));
+		}
+
+		[Fact]
+		public void WildcardStar_MatchesPrefix()
+		{
+			Assert.True(CSVStats.CsvStats.DoesSearchStringMatch("GameThread/Total", "gamethread*"));
+		}
+
+		[Fact]
+		public void WildcardStar_MatchesSuffix()
+		{
+			Assert.True(CSVStats.CsvStats.DoesSearchStringMatch("RHI/DrawPrimitive", "*/drawprimitive"));
+		}
+
+		[Fact]
+		public void WildcardMiddle_MatchesPattern()
+		{
+			Assert.True(CSVStats.CsvStats.DoesSearchStringMatch("GameThread/Total", "*thread*"));
+		}
+
+		[Fact]
+		public void MultipleWildcards_MatchesAll()
+		{
+			Assert.True(CSVStats.CsvStats.DoesSearchStringMatch("GameThread/Total/Inner", "*thread*total*"));
+		}
+
+		[Fact]
+		public void NoMatch_ReturnsFalse()
+		{
+			Assert.False(CSVStats.CsvStats.DoesSearchStringMatch("frametime", "renderthread"));
+		}
+
+		[Fact]
+		public void WildcardNoMatch_ReturnsFalse()
+		{
+			Assert.False(CSVStats.CsvStats.DoesSearchStringMatch("frametime", "render*"));
+		}
+
+		[Fact]
+		public void EmptySearchString_OnlyMatchesEmpty()
+		{
+			Assert.True(CSVStats.CsvStats.DoesSearchStringMatch("", ""));
+			Assert.False(CSVStats.CsvStats.DoesSearchStringMatch("something", ""));
+		}
+
+		[Fact]
+		public void WildcardOnly_MatchesAnything()
+		{
+			Assert.True(CSVStats.CsvStats.DoesSearchStringMatch("anything", "*"));
+			Assert.True(CSVStats.CsvStats.DoesSearchStringMatch("", "*"));
+		}
+	}
+
+	public class StatSamplesTests
+	{
+		[Fact]
+		public void ComputeAverage_ValidSamples_ReturnsCorrectAverage()
+		{
+			var stat = CsvTestHelper.CreateStatSamples("test", new float[] { 10, 20, 30, 40, 50 });
+			float avg = stat.ComputeAverage();
+			Assert.Equal(30.0f, avg, 0.001f);
+		}
+
+		[Fact]
+		public void ComputeAverage_WithRange_ReturnsCorrectAverage()
+		{
+			var stat = CsvTestHelper.CreateStatSamples("test", new float[] { 10, 20, 30, 40, 50 });
+			float avg = stat.ComputeAverage(1, 4); // 20, 30, 40
+			Assert.Equal(30.0f, avg, 0.001f);
+		}
+
+		[Fact]
+		public void ComputeMaxValue_ReturnsMax()
+		{
+			var stat = CsvTestHelper.CreateStatSamples("test", new float[] { 10, 50, 20, 40, 30 });
+			Assert.Equal(50.0f, stat.ComputeMaxValue());
+		}
+
+		[Fact]
+		public void ComputeMinValue_ReturnsMin()
+		{
+			var stat = CsvTestHelper.CreateStatSamples("test", new float[] { 10, 50, 20, 40, 30 });
+			Assert.Equal(10.0f, stat.ComputeMinValue());
+		}
+
+		[Fact]
+		public void ComputeMaxValue_WithRange_ReturnsMaxInRange()
+		{
+			var stat = CsvTestHelper.CreateStatSamples("test", new float[] { 10, 50, 20, 40, 30 });
+			Assert.Equal(50.0f, stat.ComputeMaxValue(0, 3)); // 10, 50, 20
+		}
+
+		[Fact]
+		public void GetRatioOfFramesInBudget_CorrectRatio()
+		{
+			var stat = CsvTestHelper.CreateStatSamples("frametime",
+				new float[] { 10, 15, 20, 25, 30, 35, 40, 45, 50, 55 });
+
+			float ratio = stat.GetRatioOfFramesInBudget(33.3f); // 10, 15, 20, 25, 30 = 5/10
+			Assert.Equal(0.5f, ratio, 0.001f);
+		}
+
+		[Fact]
+		public void GetRatioOfFramesOverBudget_IsComplementOfInBudget()
+		{
+			var stat = CsvTestHelper.CreateStatSamples("frametime",
+				new float[] { 10, 15, 20, 25, 30, 35, 40, 45, 50, 55 });
+
+			float inBudget = stat.GetRatioOfFramesInBudget(33.3f);
+			float overBudget = stat.GetRatioOfFramesOverBudget(33.3f);
+			Assert.Equal(1.0f, inBudget + overBudget, 0.001f);
+		}
+
+		[Fact]
+		public void GetCountOfFramesOverBudget_IgnoringFirstAndLast()
+		{
+			var stat = CsvTestHelper.CreateStatSamples("frametime",
+				new float[] { 100, 10, 20, 30, 100 });
+
+			// Explicitly pass defaults: first and last frames are ignored
+			int count = stat.GetCountOfFramesOverBudget(25, IgnoreFirstFrame: true, IgnoreLastFrame: true);
+			Assert.Equal(1, count); // Only 30 is over 25 (100s at index 0 and 4 are ignored)
+		}
+
+		[Fact]
+		public void GetCountOfFramesOverBudget_NoIgnore_CountsAll()
+		{
+			var stat = CsvTestHelper.CreateStatSamples("frametime",
+				new float[] { 100, 10, 20, 30, 100 });
+
+			int count = stat.GetCountOfFramesOverBudget(25, IgnoreFirstFrame: false, IgnoreLastFrame: false);
+			Assert.Equal(3, count); // 100 (index 0), 30 (index 3), 100 (index 4)
+		}
+
+		[Fact]
+		public void ComputeAverageAndTotal_SetsFields()
+		{
+			var stat = CsvTestHelper.CreateStatSamples("test", new float[] { 10, 20, 30 });
+			stat.ComputeAverageAndTotal();
+			Assert.Equal(20.0f, stat.average, 0.001f);
+			Assert.Equal(60.0, stat.total, 0.001);
+		}
+
+		[Fact]
+		public void GetNumSamples_ReturnsCount()
+		{
+			var stat = CsvTestHelper.CreateStatSamples("test", new float[] { 1, 2, 3, 4, 5 });
+			Assert.Equal(5, stat.GetNumSamples());
+		}
+
+		[Fact]
+		public void CompareTo_HigherAverage_IsLess()
+		{
+			var high = CsvTestHelper.CreateStatSamples("high", new float[] { 100 });
+			high.ComputeAverageAndTotal();
+			var low = CsvTestHelper.CreateStatSamples("low", new float[] { 10 });
+			low.ComputeAverageAndTotal();
+
+			// CompareTo sorts descending (higher average = earlier)
+			Assert.True(high.CompareTo(low) < 0);
+			Assert.True(low.CompareTo(high) > 0);
+		}
+	}
+}
diff --git a/Engine/Source/Programs/CSVTools/PerfReportTool.Tests/CsvStatsTests/MetadataQueryBuilderTests.cs b/Engine/Source/Programs/CSVTools/PerfReportTool.Tests/CsvStatsTests/MetadataQueryBuilderTests.cs
new file mode 100644
--- /dev/null
+++ b/Engine/Source/Programs/CSVTools/PerfReportTool.Tests/CsvStatsTests/MetadataQueryBuilderTests.cs
@@ -0,0 +1,212 @@
+using CSVStats;
+using PerfReportTool.Tests.TestHelpers;
+using Xunit;
+
+namespace PerfReportTool.Tests.CsvStatsTests
+{
+	public class MetadataQueryBuilderTests
+	{
+		[Fact]
+		public void BuildQueryExpressionTree_SimpleEquality_Works()
+		{
+			var expr = MetadataQueryBuilder.BuildQueryExpressionTree("platform=windows");
+			Assert.NotNull(expr);
+
+			var metadata = CsvTestHelper.CreateMetadata(("platform", "windows"));
+			Assert.True(expr.Evaluate(metadata));
+		}
+
+		[Fact]
+		public void BuildQueryExpressionTree_SimpleEquality_NoMatch()
+		{
+			var expr = MetadataQueryBuilder.BuildQueryExpressionTree("platform=windows");
+			var metadata = CsvTestHelper.CreateMetadata(("platform", "linux"));
+			Assert.False(expr.Evaluate(metadata));
+		}
+
+		[Fact]
+		public void BuildQueryExpressionTree_NotEqual_Works()
+		{
+			var expr = MetadataQueryBuilder.BuildQueryExpressionTree("platform!=windows");
+			var metadata = CsvTestHelper.CreateMetadata(("platform", "linux"));
+			Assert.True(expr.Evaluate(metadata));
+
+			var metadata2 = CsvTestHelper.CreateMetadata(("platform", "windows"));
+			Assert.False(expr.Evaluate(metadata2));
+		}
+
+		[Fact]
+		public void BuildQueryExpressionTree_And_CombinesCorrectly()
+		{
+			var expr = MetadataQueryBuilder.BuildQueryExpressionTree("platform=windows and config=test");
+
+			var both = CsvTestHelper.CreateMetadata(("platform", "windows"), ("config", "test"));
+			Assert.True(expr.Evaluate(both));
+
+			var onlyPlatform = CsvTestHelper.CreateMetadata(("platform", "windows"), ("config", "shipping"));
+			Assert.False(expr.Evaluate(onlyPlatform));
+		}
+
+		[Fact]
+		public void BuildQueryExpressionTree_Or_CombinesCorrectly()
+		{
+			var expr = MetadataQueryBuilder.BuildQueryExpressionTree("platform=windows or platform=linux");
+
+			var win = CsvTestHelper.CreateMetadata(("platform", "windows"));
+			Assert.True(expr.Evaluate(win));
+
+			var linux = CsvTestHelper.CreateMetadata(("platform", "linux"));
+			Assert.True(expr.Evaluate(linux));
+
+			var mac = CsvTestHelper.CreateMetadata(("platform", "mac"));
+			Assert.False(expr.Evaluate(mac));
+		}
+
+		[Fact]
+		public void BuildQueryExpressionTree_NotOperator_NegatesExpression()
+		{
+			var expr = MetadataQueryBuilder.BuildQueryExpressionTree("not(platform=windows)");
+
+			var win = CsvTestHelper.CreateMetadata(("platform", "windows"));
+			Assert.False(expr.Evaluate(win));
+
+			var linux = CsvTestHelper.CreateMetadata(("platform", "linux"));
+			Assert.True(expr.Evaluate(linux));
+		}
+
+		[Fact]
+		public void BuildQueryExpressionTree_NestedParentheses_ParsesCorrectly()
+		{
+			var expr = MetadataQueryBuilder.BuildQueryExpressionTree(
+				"config=test and (platform=windows or platform=linux)");
+
+			var winTest = CsvTestHelper.CreateMetadata(("config", "test"), ("platform", "windows"));
+			Assert.True(expr.Evaluate(winTest));
+
+			var linuxTest = CsvTestHelper.CreateMetadata(("config", "test"), ("platform", "linux"));
+			Assert.True(expr.Evaluate(linuxTest));
+
+			var macTest = CsvTestHelper.CreateMetadata(("config", "test"), ("platform", "mac"));
+			Assert.False(expr.Evaluate(macTest));
+
+			var winShipping = CsvTestHelper.CreateMetadata(("config", "shipping"), ("platform", "windows"));
+			Assert.False(expr.Evaluate(winShipping));
+		}
+
+		[Fact]
+		public void BuildQueryExpressionTree_ComparisonOperators_Work()
+		{
+			var ltExpr = MetadataQueryBuilder.BuildQueryExpressionTree("fps<60");
+			var gtExpr = MetadataQueryBuilder.BuildQueryExpressionTree("fps>30");
+			var leExpr = MetadataQueryBuilder.BuildQueryExpressionTree("fps<=60");
+			var geExpr = MetadataQueryBuilder.BuildQueryExpressionTree("fps>=30");
+
+			var meta45 = CsvTestHelper.CreateMetadata(("fps", "45"));
+			Assert.True(ltExpr.Evaluate(meta45));
+			Assert.True(gtExpr.Evaluate(meta45));
+			Assert.True(leExpr.Evaluate(meta45));
+			Assert.True(geExpr.Evaluate(meta45));
+
+			var meta60 = CsvTestHelper.CreateMetadata(("fps", "60"));
+			Assert.False(ltExpr.Evaluate(meta60));
+			Assert.True(gtExpr.Evaluate(meta60));
+			Assert.True(leExpr.Evaluate(meta60));
+			Assert.True(geExpr.Evaluate(meta60));
+		}
+
+		[Fact]
+		public void BuildQueryExpressionTree_WildcardValue_MatchesPattern()
+		{
+			var expr = MetadataQueryBuilder.BuildQueryExpressionTree("deviceprofile=windows*");
+
+			var win = CsvTestHelper.CreateMetadata(("deviceprofile", "windows"));
+			Assert.True(expr.Evaluate(win));
+
+			var winHigh = CsvTestHelper.CreateMetadata(("deviceprofile", "windows_high"));
+			Assert.True(expr.Evaluate(winHigh));
+
+			var linux = CsvTestHelper.CreateMetadata(("deviceprofile", "linux"));
+			Assert.False(expr.Evaluate(linux));
+		}
+
+		[Fact]
+		public void BuildQueryExpressionTree_ContainsOperator_Works()
+		{
+			var expr = MetadataQueryBuilder.BuildQueryExpressionTree("buildversion~=main");
+
+			var main = CsvTestHelper.CreateMetadata(("buildversion", "++ue5+main-cl-12345"));
+			Assert.True(expr.Evaluate(main));
+
+			var release = CsvTestHelper.CreateMetadata(("buildversion", "++ue5+release-5.4"));
+			Assert.False(expr.Evaluate(release));
+		}
+
+		[Fact]
+		public void BuildQueryExpressionTree_CommaNotation_ConvertsToAnd()
+		{
+			// Old-style comma-separated filters should be treated as AND
+			var expr = MetadataQueryBuilder.BuildQueryExpressionTree("platform=windows,config=test");
+
+			var match = CsvTestHelper.CreateMetadata(("platform", "windows"), ("config", "test"));
+			Assert.True(expr.Evaluate(match));
+
+			var noMatch = CsvTestHelper.CreateMetadata(("platform", "windows"), ("config", "shipping"));
+			Assert.False(expr.Evaluate(noMatch));
+		}
+
+		[Fact]
+		public void Evaluate_ComplexFilter_MatchesExpectedResults()
+		{
+			var expr = MetadataQueryBuilder.BuildQueryExpressionTree(
+				"buildversion=++UE5+Main-CL-52441555 and (platform=windows or platform=linux) and not(deviceprofile=windows_low* or deviceprofile=linux_low*)");
+
+			// Should match: Windows with correct build, non-low profile
+			var winHigh = CsvTestHelper.CreateMetadata(
+				("buildversion", "++ue5+main-cl-52441555"),
+				("platform", "windows"),
+				("deviceprofile", "windows"));
+			Assert.True(expr.Evaluate(winHigh));
+
+			// Should NOT match: Windows low profile (excluded by not clause)
+			var winLow = CsvTestHelper.CreateMetadata(
+				("buildversion", "++ue5+main-cl-52441555"),
+				("platform", "windows"),
+				("deviceprofile", "windows_low"));
+			Assert.False(expr.Evaluate(winLow));
+
+			// Should NOT match: wrong build version
+			var wrongBuild = CsvTestHelper.CreateMetadata(
+				("buildversion", "++ue5+release-5.4-cl-99999"),
+				("platform", "windows"),
+				("deviceprofile", "windows"));
+			Assert.False(expr.Evaluate(wrongBuild));
+
+			// Should NOT match: Mac platform (not in the or list)
+			var mac = CsvTestHelper.CreateMetadata(
+				("buildversion", "++ue5+main-cl-52441555"),
+				("platform", "mac"),
+				("deviceprofile", "mac"));
+			Assert.False(expr.Evaluate(mac));
+		}
+
+		[Fact]
+		public void Evaluate_MissingMetadataKey_ReturnsFalse()
+		{
+			var expr = MetadataQueryBuilder.BuildQueryExpressionTree("platform=windows");
+			var emptyMeta = CsvTestHelper.CreateMetadata();
+			Assert.False(expr.Evaluate(emptyMeta));
+		}
+
+		[Theory]
+		[InlineData("platform=windows")]
+		[InlineData("a=1 and b=2")]
+		[InlineData("a=1 or b=2")]
+		[InlineData("not(a=1)")]
+		[InlineData("(a=1 or b=2) and c=3")]
+		public void BuildQueryExpressionTree_ValidExpressions_DoNotThrow(string filter)
+		{
+			var expr = MetadataQueryBuilder.BuildQueryExpressionTree(filter);
+			Assert.NotNull(expr);
+		}
+	}
+}
diff --git a/Engine/Source/Programs/CSVTools/PerfReportTool.Tests/Integration/CliArgumentTests.cs b/Engine/Source/Programs/CSVTools/PerfReportTool.Tests/Integration/CliArgumentTests.cs
new file mode 100644
--- /dev/null
+++ b/Engine/Source/Programs/CSVTools/PerfReportTool.Tests/Integration/CliArgumentTests.cs
@@ -0,0 +1,874 @@
+// Copyright Epic Games, Inc. All Rights Reserved.
+
+using System;
+using System.Collections.Generic;
+using System.Diagnostics;
+using System.IO;
+using System.Linq;
+using System.Text.Json;
+using Xunit;
+
+namespace PerfReportTool.Tests.Integration
+{
+	[Trait("Category", "Integration")]
+	[Collection("EndToEnd")]
+	public class CliArgumentTests : IClassFixture<TestDataFixture>, IDisposable
+	{
+		private readonly TestDataFixture _fixture;
+		private readonly string _outputDir;
+		private readonly string _toolPath;
+
+		public CliArgumentTests(TestDataFixture fixture)
+		{
+			_fixture = fixture;
+			_outputDir = Path.Combine(Path.GetTempPath(), "PerfReportToolCli_" + Path.GetRandomFileName());
+			Directory.CreateDirectory(_outputDir);
+
+			string assemblyDir = Path.GetDirectoryName(typeof(PerfReportTool.Version).Assembly.Location);
+			_toolPath = Path.Combine(assemblyDir, "PerfreportTool.dll");
+		}
+
+		public void Dispose()
+		{
+			if (Directory.Exists(_outputDir))
+			{
+				try { Directory.Delete(_outputDir, true); } catch { }
+			}
+		}
+
+		private (int exitCode, string output) RunTool(string arguments, int timeoutMs = 300000)
+		{
+			var psi = new ProcessStartInfo
+			{
+				FileName = "dotnet",
+				Arguments = $"\"{_toolPath}\" {arguments}",
+				RedirectStandardOutput = true,
+				RedirectStandardError = true,
+				UseShellExecute = false,
+				CreateNoWindow = true
+			};
+
+			using var process = Process.Start(psi);
+			string stdout = process.StandardOutput.ReadToEnd();
+			string stderr = process.StandardError.ReadToEnd();
+			if (!process.WaitForExit(timeoutMs))
+			{
+				try { process.Kill(); } catch { }
+				return (-1, stdout + stderr + "\n[ERROR] Process timed out");
+			}
+
+			return (process.ExitCode, stdout + stderr);
+		}
+
+		/// <summary>Common args for bulk mode with email table output.</summary>
+		private string BulkArgs(string outputDir)
+		{
+			return $"-csvdir \"{_fixture.GetPerfCsvDir()}\" -emailtable -recurse -o \"{outputDir}\" " +
+				$"-nowatermarks -nodetailedreports -searchpattern Perf_*.csv.bin";
+		}
+
+		private string GeneratePrcCache(string name)
+		{
+			string cacheDir = Path.Combine(_outputDir, name + "_cache");
+			Directory.CreateDirectory(cacheDir);
+			string outputDir = Path.Combine(_outputDir, name + "_html");
+
+			var (exitCode, output) = RunTool(
+				$"{BulkArgs(outputDir)} " +
+				$"-summaryTableCacheUseOnlyCsvID -summaryTableCacheInvalidate " +
+				$"-summaryTableCache \"{cacheDir}\"");
+
+			Assert.DoesNotContain("[ERROR]", output);
+			Assert.True(Directory.GetFiles(cacheDir, "*.prc").Length > 0,
+				$"PRC cache generation failed — no .prc files in {cacheDir}. Output: {output.Substring(0, Math.Min(500, output.Length))}");
+
+			return cacheDir;
+		}
+
+		/// <summary>Read the first HTML summary table file from a directory.</summary>
+		private string ReadFirstHtml(string dir, string pattern = "*.html")
+		{
+			Assert.True(Directory.Exists(dir), $"Output directory does not exist: {dir}");
+			string[] files = Directory.GetFiles(dir, pattern, SearchOption.TopDirectoryOnly);
+			Assert.True(files.Length > 0,
+				$"No '{pattern}' files found in {dir}. Tool may have failed silently — check tool output for errors.");
+			return File.ReadAllText(files[0]);
+		}
+
+		#region Data Filtering
+
+		[Fact]
+		public void MinX_MaxX_FrameTruncation_AffectsOutput()
+		{
+			_fixture.RequireTestData();
+			string outputFull = Path.Combine(_outputDir, "MinMaxX_Full");
+			string outputCropped = Path.Combine(_outputDir, "MinMaxX_Cropped");
+
+			RunTool(BulkArgs(outputFull));
+			RunTool($"{BulkArgs(outputCropped)} -minx 10 -maxx 50");
+
+			string htmlFull = ReadFirstHtml(outputFull);
+			string htmlCropped = ReadFirstHtml(outputCropped);
+
+			Assert.NotNull(htmlFull);
+			Assert.NotNull(htmlCropped);
+			// Cropped output should differ from full output (different stats computed over fewer frames)
+			Assert.NotEqual(htmlFull, htmlCropped);
+		}
+
+		[Fact]
+		public void NoStripEvents_ProducesOutput()
+		{
+			_fixture.RequireTestData();
+			string outputDir = Path.Combine(_outputDir, "NoStrip");
+
+			var (exitCode, output) = RunTool($"{BulkArgs(outputDir)} -nostripevents");
+
+			Assert.DoesNotContain("[ERROR]", output);
+			Assert.True(Directory.GetFiles(outputDir, "*.html").Length > 0, "Should produce HTML output");
+		}
+
+		#endregion
+
+		#region Summary Table Format
+
+		[Fact]
+		public void CollateTable_ProducesCollatedHtmlFile()
+		{
+			_fixture.RequireTestData();
+			string outputDir = Path.Combine(_outputDir, "Collate");
+
+			var (exitCode, output) = RunTool($"{BulkArgs(outputDir)} -collatetable");
+
+			Assert.DoesNotContain("[ERROR]", output);
+			Assert.True(Directory.Exists(outputDir));
+			string[] collatedFiles = Directory.GetFiles(outputDir, "*Collated*", SearchOption.TopDirectoryOnly);
+			Assert.True(collatedFiles.Length > 0, "Collate table should produce a *Collated* file");
+		}
+
+		[Fact]
+		public void CsvTable_ProducesCsvNotHtml()
+		{
+			_fixture.RequireTestData();
+			string outputDir = Path.Combine(_outputDir, "CsvTable");
+
+			var (exitCode, output) = RunTool($"{BulkArgs(outputDir)} -csvtable");
+
+			Assert.DoesNotContain("[ERROR]", output);
+			Assert.True(Directory.Exists(outputDir));
+			string[] csvFiles = Directory.GetFiles(outputDir, "*.csv", SearchOption.TopDirectoryOnly);
+			Assert.True(csvFiles.Length > 0, "CsvTable should produce CSV output files");
+			// CSV file should contain comma-separated data with a header row
+			string csvContent = File.ReadAllText(csvFiles[0]);
+			Assert.Contains(",", csvContent);
+		}
+
+		[Fact]
+		public void TransposeTable_OutputContainsTransposeMarker()
+		{
+			_fixture.RequireTestData();
+			string outputDir = Path.Combine(_outputDir, "Transpose");
+
+			var (exitCode, output) = RunTool($"{BulkArgs(outputDir)} -transposetable");
+
+			Assert.DoesNotContain("[ERROR]", output);
+			string html = ReadFirstHtml(outputDir);
+			Assert.NotNull(html);
+			// Transposed tables embed SummaryTableIsTransposed=true in JavaScript
+			Assert.Contains("SummaryTableIsTransposed", html);
+		}
+
+		[Fact]
+		public void ReverseTable_ReversesRowOrder()
+		{
+			_fixture.RequireTestData();
+			string outputNormal = Path.Combine(_outputDir, "ReverseNormal");
+			string outputReversed = Path.Combine(_outputDir, "ReverseReversed");
+
+			RunTool(BulkArgs(outputNormal));
+			RunTool($"{BulkArgs(outputReversed)} -reversetable");
+
+			string htmlNormal = ReadFirstHtml(outputNormal);
+			string htmlReversed = ReadFirstHtml(outputReversed);
+
+			Assert.NotNull(htmlNormal);
+			Assert.NotNull(htmlReversed);
+			// Reversed table should produce different HTML (different row ordering)
+			Assert.NotEqual(htmlNormal, htmlReversed);
+		}
+
+		[Fact]
+		public void ScrollableTable_EmbedsScrollableJavaScript()
+		{
+			_fixture.RequireTestData();
+			string outputDir = Path.Combine(_outputDir, "Scrollable");
+
+			var (exitCode, output) = RunTool($"{BulkArgs(outputDir)} -scrollabletable");
+
+			Assert.DoesNotContain("[ERROR]", output);
+			string html = ReadFirstHtml(outputDir);
+			Assert.NotNull(html);
+			Assert.Contains("regenerateStickyColumnCss", html);
+		}
+
+		[Fact]
+		public void ColorizeTable_Off_ProducesNoBackgroundColors()
+		{
+			_fixture.RequireTestData();
+			string outputDir = Path.Combine(_outputDir, "ColorOff");
+
+			var (exitCode, output) = RunTool($"{BulkArgs(outputDir)} -colorizetable off");
+
+			Assert.DoesNotContain("[ERROR]", output);
+			string html = ReadFirstHtml(outputDir);
+			Assert.NotNull(html);
+			// With colorize off, cells should not have bgcolor attributes from thresholds
+			// (there may still be some from other sources, but far fewer)
+			Assert.DoesNotContain("bgcolor='#ff", html.ToLower());
+		}
+
+		[Fact]
+		public void ColorizeTable_Auto_ProducesBackgroundColors()
+		{
+			_fixture.RequireTestData();
+			string outputDir = Path.Combine(_outputDir, "ColorAuto");
+
+			var (exitCode, output) = RunTool($"{BulkArgs(outputDir)} -colorizetable auto");
+
+			Assert.DoesNotContain("[ERROR]", output);
+			string html = ReadFirstHtml(outputDir);
+			Assert.NotNull(html);
+			// Auto colorize should produce bgcolor attributes on cells
+			Assert.Contains("bgcolor", html.ToLower());
+		}
+
+		[Fact]
+		public void SpreadsheetFriendly_PrefixesStringValues()
+		{
+			_fixture.RequireTestData();
+			string outputDir = Path.Combine(_outputDir, "Spreadsheet");
+
+			var (exitCode, output) = RunTool($"{BulkArgs(outputDir)} -spreadsheetfriendly");
+
+			Assert.DoesNotContain("[ERROR]", output);
+			string html = ReadFirstHtml(outputDir);
+			Assert.NotNull(html);
+			// Spreadsheet-friendly mode prefixes non-numeric text entries with a single quote
+			// Look for the quote prefix pattern in table cells: >'text
+			Assert.Contains(">'", html);
+		}
+
+		[Fact]
+		public void NoSummaryMinMax_ProducesOutput()
+		{
+			_fixture.RequireTestData();
+			string outputDir = Path.Combine(_outputDir, "NoMinMax");
+
+			var (exitCode, output) = RunTool($"{BulkArgs(outputDir)} -nosummaryminmax");
+
+			Assert.DoesNotContain("[ERROR]", output);
+			Assert.True(Directory.GetFiles(outputDir, "*.html").Length > 0);
+		}
+
+		[Fact]
+		public void SummaryTableFilename_UsesCustomName()
+		{
+			_fixture.RequireTestData();
+			string outputDir = Path.Combine(_outputDir, "CustomFilename");
+
+			var (exitCode, output) = RunTool($"{BulkArgs(outputDir)} -summarytablefilename CustomTable.html");
+
+			Assert.DoesNotContain("[ERROR]", output);
+			Assert.True(Directory.Exists(outputDir));
+			string[] customFiles = Directory.GetFiles(outputDir, "CustomTable*");
+			Assert.True(customFiles.Length > 0, "Output should use the custom filename 'CustomTable'");
+			// The default name "SummaryTable" should NOT appear
+			string[] defaultFiles = Directory.GetFiles(outputDir, "SummaryTable*");
+			Assert.True(defaultFiles.Length == 0, "Default 'SummaryTable' filename should not be used when custom name is specified");
+		}
+
+		#endregion
+
+		#region Diff Rows
+
+		[Fact]
+		public void AddDiffRows_InsertsAdditionalRows()
+		{
+			_fixture.RequireTestData();
+			string outputNoDiff = Path.Combine(_outputDir, "DiffRowsNone");
+			string outputWithDiff = Path.Combine(_outputDir, "DiffRowsWith");
+
+			RunTool(BulkArgs(outputNoDiff));
+			RunTool($"{BulkArgs(outputWithDiff)} -adddiffrows");
+
+			string htmlNoDiff = ReadFirstHtml(outputNoDiff);
+			string htmlWithDiff = ReadFirstHtml(outputWithDiff);
+
+			Assert.NotNull(htmlNoDiff);
+			Assert.NotNull(htmlWithDiff);
+			// Diff rows version should have more <tr> elements
+			int rowCountNoDiff = htmlNoDiff.Split("<tr").Length;
+			int rowCountWithDiff = htmlWithDiff.Split("<tr").Length;
+			Assert.True(rowCountWithDiff >= rowCountNoDiff,
+				$"Diff rows should add rows: {rowCountNoDiff} without vs {rowCountWithDiff} with");
+		}
+
+		[Fact]
+		public void AddDiffRows_WithSortByDiff_ProducesOutput()
+		{
+			_fixture.RequireTestData();
+			string outputDir = Path.Combine(_outputDir, "DiffRowsSort");
+
+			var (exitCode, output) = RunTool($"{BulkArgs(outputDir)} -adddiffrows -sortcolumnsbydiff");
+
+			Assert.DoesNotContain("[ERROR]", output);
+			Assert.True(Directory.GetFiles(outputDir, "*.html").Length > 0);
+		}
+
+		[Fact]
+		public void ShowOnlyDiffRows_OutputsFewerRows()
+		{
+			_fixture.RequireTestData();
+			string outputWithDiff = Path.Combine(_outputDir, "DiffRowsFull");
+			string outputOnlyDiff = Path.Combine(_outputDir, "OnlyDiffRows");
+
+			RunTool($"{BulkArgs(outputWithDiff)} -adddiffrows");
+			RunTool($"{BulkArgs(outputOnlyDiff)} -adddiffrows -showonlydiffrows");
+
+			string htmlFull = ReadFirstHtml(outputWithDiff);
+			string htmlOnly = ReadFirstHtml(outputOnlyDiff);
+
+			Assert.NotNull(htmlFull);
+			Assert.NotNull(htmlOnly);
+			// ShowOnlyDiffRows skips data rows entirely, so output should have fewer <tr> elements
+			int rowCountFull = htmlFull.Split("<tr").Length;
+			int rowCountOnly = htmlOnly.Split("<tr").Length;
+			Assert.True(rowCountOnly < rowCountFull,
+				$"ShowOnlyDiffRows should produce fewer rows: {rowCountFull} full vs {rowCountOnly} only-diff");
+		}
+
+		#endregion
+
+		#region Stat Reading Modes
+
+		[Fact]
+		public void ReadStatMinMax_AddsMinMaxColumns()
+		{
+			_fixture.RequireTestData();
+			string outputNormal = Path.Combine(_outputDir, "ReadStatNormal");
+			string outputMinMax = Path.Combine(_outputDir, "ReadStatMinMax");
+
+			RunTool(BulkArgs(outputNormal));
+			RunTool($"{BulkArgs(outputMinMax)} -readstatminmax");
+
+			string htmlNormal = ReadFirstHtml(outputNormal);
+			string htmlMinMax = ReadFirstHtml(outputMinMax);
+
+			Assert.NotNull(htmlNormal);
+			Assert.NotNull(htmlMinMax);
+			// MinMax mode should produce a larger table (more columns)
+			Assert.True(htmlMinMax.Length > htmlNormal.Length,
+				"ReadStatMinMax should produce more output with additional min/max columns");
+		}
+
+		[Fact]
+		public void ReadStatAvg_ProducesOutput()
+		{
+			_fixture.RequireTestData();
+			string outputDir = Path.Combine(_outputDir, "ReadStatAvg");
+
+			var (exitCode, output) = RunTool($"{BulkArgs(outputDir)} -readstatavg");
+
+			Assert.DoesNotContain("[ERROR]", output);
+			Assert.True(Directory.GetFiles(outputDir, "*.html").Length > 0);
+		}
+
+		#endregion
+
+		#region No Reports / NoDetailedReports
+
+		[Fact(Skip = "BUG: -noreports does not gate WriteSummaryTableReport, causing WriteSummaryTableReport to run on a table with rows but no columns (GenerateReport, which populates row data, is skipped). Crashes in WildcardMatchStringList binary search on empty list.")]
+		public void NoReports_SkipsAllReports_StillGeneratesCache()
+		{
+			_fixture.RequireTestData();
+			string outputDir = Path.Combine(_outputDir, "NoReports");
+			string cacheDir = Path.Combine(_outputDir, "NoReportsCache");
+			Directory.CreateDirectory(cacheDir);
+
+			var (exitCode, output) = RunTool(
+				$"{BulkArgs(outputDir)} " +
+				$"-noreports " +
+				$"-summaryTableCacheUseOnlyCsvID -summaryTableCacheInvalidate " +
+				$"-summaryTableCache \"{cacheDir}\"");
+
+			Assert.DoesNotContain("[ERROR]", output);
+		}
+
+		#endregion
+
+		#region Variable Dumping
+
+		[Fact]
+		public void DumpVariables_PrintsVariablesToStdout()
+		{
+			_fixture.RequireTestData();
+			string outputDir = Path.Combine(_outputDir, "DumpVars");
+			string csvDir = _fixture.GetPerfCsvDir();
+			string[] csvFiles = Directory.GetFiles(csvDir, "Perf_*.csv.bin");
+			Assert.True(csvFiles.Length > 0, $"No Perf_*.csv.bin files found in {csvDir}");
+
+			var (exitCode, output) = RunTool(
+				$"-csv \"{csvFiles[0]}\" -o \"{outputDir}\" " +
+				$"-nowatermarks -nodetailedreports -dumpvariables");
+
+			Assert.DoesNotContain("[ERROR]", output);
+			// -dumpvariables should print variable names and values to stdout
+			// Variables typically include budget values from ReportTypes.xml
+			Assert.True(output.Length > 100, "DumpVariables should produce substantial stdout output");
+		}
+
+		#endregion
+
+		#region ListSummaryTables
+
+		[Fact]
+		public void ListSummaryTables_ListsTableNames()
+		{
+			var (exitCode, output) = RunTool("-listsummarytables");
+			Assert.DoesNotContain("[ERROR]", output);
+			// Should print "Listing summary tables:" followed by table names
+			Assert.Contains("summary table", output.ToLower());
+		}
+
+		#endregion
+
+		#region Collation Options
+
+		[Fact]
+		public void CollateTableOnly_ProducesOnlyCollatedOutput()
+		{
+			_fixture.RequireTestData();
+			string outputDir = Path.Combine(_outputDir, "CollateOnly");
+
+			var (exitCode, output) = RunTool($"{BulkArgs(outputDir)} -collatetableonly");
+
+			Assert.DoesNotContain("[ERROR]", output);
+			Assert.True(Directory.Exists(outputDir));
+			// With -collatetableonly, should NOT produce a non-collated summary table
+			string[] htmlFiles = Directory.GetFiles(outputDir, "*.html", SearchOption.TopDirectoryOnly);
+			Assert.True(htmlFiles.Length > 0, "Should produce at least one HTML file");
+		}
+
+		[Fact]
+		public void NoWeightedAvg_ProducesOutput()
+		{
+			_fixture.RequireTestData();
+			string outputDir = Path.Combine(_outputDir, "NoWeight");
+
+			var (exitCode, output) = RunTool($"{BulkArgs(outputDir)} -collatetable -noweightedavg");
+
+			Assert.DoesNotContain("[ERROR]", output);
+			Assert.True(Directory.GetFiles(outputDir, "*.html").Length > 0);
+		}
+
+		#endregion
+
+		#region Column Sort Modes
+
+		[Theory]
+		[InlineData("sortbymax")]
+		[InlineData("sortbyavg")]
+		[InlineData("wildcardSortByMax")]
+		[InlineData("wildcardSortByAvg")]
+		public void ColumnSortMode_ProducesOutput(string mode)
+		{
+			_fixture.RequireTestData();
+			string outputDir = Path.Combine(_outputDir, $"ColSort_{mode}");
+
+			var (exitCode, output) = RunTool($"{BulkArgs(outputDir)} -columnsortmode {mode}");
+
+			Assert.DoesNotContain("[ERROR]", output);
+			Assert.True(Directory.GetFiles(outputDir, "*.html").Length > 0,
+				$"Column sort mode '{mode}' should produce HTML output");
+		}
+
+		#endregion
+
+		#region Override Metadata
+
+		[Fact]
+		public void OverrideMetadata_AcceptedWithoutError()
+		{
+			_fixture.RequireTestData();
+			string outputDir = Path.Combine(_outputDir, "OverrideMeta");
+
+			// -overrideMetadata applies to the -metadataProxy path, not bulk CSV mode.
+			// Verify the argument is accepted and processed without error in bulk mode.
+			var (exitCode, output) = RunTool(
+				$"{BulkArgs(outputDir)} -overridemetadata platform=TestOverride,targetframerate=60");
+
+			Assert.DoesNotContain("[ERROR]", output);
+			Assert.True(Directory.GetFiles(outputDir, "*.html").Length > 0);
+		}
+
+		#endregion
+
+		#region Cache Options
+
+		[Fact]
+		public void SummaryTableCacheReadOnly_DoesNotWriteNewFiles()
+		{
+			_fixture.RequireTestData();
+			string cacheDir = GeneratePrcCache("CacheRO");
+			int prcCountBefore = Directory.GetFiles(cacheDir, "*.prc").Length;
+
+			string outputDir = Path.Combine(_outputDir, "CacheRO_Out");
+			var (exitCode, output) = RunTool(
+				$"{BulkArgs(outputDir)} " +
+				$"-summaryTableCacheUseOnlyCsvID " +
+				$"-summaryTableCacheReadOnly " +
+				$"-summaryTableCache \"{cacheDir}\"");
+
+			Assert.DoesNotContain("[ERROR]", output);
+			int prcCountAfter = Directory.GetFiles(cacheDir, "*.prc").Length;
+			Assert.Equal(prcCountBefore, prcCountAfter);
+		}
+
+		[Fact]
+		public void SummaryTableCachePurgeInvalid_ProducesOutput()
+		{
+			_fixture.RequireTestData();
+			string cacheDir = GeneratePrcCache("CachePurge");
+			string outputDir = Path.Combine(_outputDir, "CachePurge_Out");
+
+			var (exitCode, output) = RunTool(
+				$"{BulkArgs(outputDir)} " +
+				$"-summaryTableCacheUseOnlyCsvID " +
+				$"-summaryTableCachePurgeInvalid " +
+				$"-summaryTableCache \"{cacheDir}\"");
+
+			Assert.DoesNotContain("[ERROR]", output);
+			Assert.True(Directory.GetFiles(outputDir, "*.html").Length > 0);
+		}
+
+		#endregion
+
+		#region Performance/Threading Options
+
+		[Fact]
+		public void CsvToSvgSequential_ProducesIdenticalOutput()
+		{
+			_fixture.RequireTestData();
+			string outputDir = Path.Combine(_outputDir, "Sequential");
+
+			var (exitCode, output) = RunTool($"{BulkArgs(outputDir)} -csvtosvgsequential");
+
+			Assert.DoesNotContain("[ERROR]", output);
+			Assert.True(Directory.GetFiles(outputDir, "*.html").Length > 0);
+		}
+
+		[Fact]
+		public void GraphThreads_CustomValue_ProducesOutput()
+		{
+			_fixture.RequireTestData();
+			string outputDir = Path.Combine(_outputDir, "GraphThreads");
+
+			var (exitCode, output) = RunTool($"{BulkArgs(outputDir)} -graphthreads 1");
+
+			Assert.DoesNotContain("[ERROR]", output);
+			Assert.True(Directory.GetFiles(outputDir, "*.html").Length > 0);
+		}
+
+		[Fact]
+		public void PrecacheSettings_CustomValues_ProducesOutput()
+		{
+			_fixture.RequireTestData();
+			string outputDir = Path.Combine(_outputDir, "Precache");
+
+			var (exitCode, output) = RunTool($"{BulkArgs(outputDir)} -precachecount 2 -precachethreads 2");
+
+			Assert.DoesNotContain("[ERROR]", output);
+			Assert.True(Directory.GetFiles(outputDir, "*.html").Length > 0);
+		}
+
+		#endregion
+
+		#region Stat Threshold Column Filter
+
+		[Fact]
+		public void SummaryTableStatThreshold_ReducesColumnCount()
+		{
+			_fixture.RequireTestData();
+			string outputNormal = Path.Combine(_outputDir, "StatThresholdNormal");
+			string outputFiltered = Path.Combine(_outputDir, "StatThresholdFiltered");
+
+			RunTool(BulkArgs(outputNormal));
+			// High threshold should filter out many low-value stat columns
+			RunTool($"{BulkArgs(outputFiltered)} -summarytablestatthreshold 999999");
+
+			string htmlNormal = ReadFirstHtml(outputNormal);
+			string htmlFiltered = ReadFirstHtml(outputFiltered);
+
+			Assert.NotNull(htmlNormal);
+			Assert.NotNull(htmlFiltered);
+			// Filtered output should have fewer <th> headers (fewer columns)
+			int colCountNormal = htmlNormal.Split("<th").Length;
+			int colCountFiltered = htmlFiltered.Split("<th").Length;
+			Assert.True(colCountFiltered < colCountNormal,
+				$"High stat threshold should filter out columns: {colCountNormal} normal vs {colCountFiltered} filtered");
+		}
+
+		#endregion
+
+		#region Hide Metadata Columns
+
+		[Fact]
+		public void HideMetadataColumns_RemovesMetadataFromTable()
+		{
+			_fixture.RequireTestData();
+			string outputNormal = Path.Combine(_outputDir, "HideMetaNormal");
+			string outputHidden = Path.Combine(_outputDir, "HideMetaHidden");
+
+			RunTool(BulkArgs(outputNormal));
+			RunTool($"{BulkArgs(outputHidden)} -hidemetadatacolumns");
+
+			string htmlNormal = ReadFirstHtml(outputNormal);
+			string htmlHidden = ReadFirstHtml(outputHidden);
+
+			Assert.NotNull(htmlNormal);
+			Assert.NotNull(htmlHidden);
+			// Hidden metadata output should be smaller (fewer columns)
+			int colCountNormal = htmlNormal.Split("<th").Length;
+			int colCountHidden = htmlHidden.Split("<th").Length;
+			Assert.True(colCountHidden < colCountNormal,
+				$"Hiding metadata should reduce columns: {colCountNormal} normal vs {colCountHidden} hidden");
+		}
+
+		#endregion
+
+		#region Response File
+
+		[Fact]
+		public void ResponseFile_ProcessesCorrectly()
+		{
+			_fixture.RequireTestData();
+			string outputDir = Path.Combine(_outputDir, "ResponseFile");
+			string responseFile = Path.Combine(_outputDir, "response.txt");
+			string csvDir = _fixture.GetPerfCsvDir();
+
+			File.WriteAllLines(responseFile, new[]
+			{
+				$"-csvdir \"{csvDir}\"",
+				"-emailtable",
+				"-recurse",
+				$"-o \"{outputDir}\"",
+				"-nowatermarks",
+				"-nodetailedreports",
+				"-searchpattern Perf_*.csv.bin"
+			});
+
+			var (exitCode, output) = RunTool($"-response \"{responseFile}\"");
+
+			Assert.DoesNotContain("[ERROR]", output);
+			Assert.True(Directory.Exists(outputDir));
+			Assert.True(Directory.GetFiles(outputDir, "*.html").Length > 0,
+				"Response file should produce identical output to direct arguments");
+		}
+
+		#endregion
+
+		#region SummaryTable JS Data
+
+		[Fact]
+		public void SummaryTableJsData_EmbedsJavaScriptData()
+		{
+			_fixture.RequireTestData();
+			string outputDir = Path.Combine(_outputDir, "JsData");
+
+			var (exitCode, output) = RunTool($"{BulkArgs(outputDir)} -summarytablejsdata");
+
+			Assert.DoesNotContain("[ERROR]", output);
+			string html = ReadFirstHtml(outputDir);
+			Assert.NotNull(html);
+			// JS data mode embeds row/column data in a <script> section
+			Assert.Contains("<script>", html);
+			Assert.Contains("elementType", html);
+		}
+
+		#endregion
+
+		#region Max Summary Table String Length
+
+		[Fact]
+		public void MaxSummaryTableStringLength_TruncatesLongStrings()
+		{
+			_fixture.RequireTestData();
+			string outputDir = Path.Combine(_outputDir, "MaxStrLen");
+
+			// Use a very short max length to force truncation
+			var (exitCode, output) = RunTool($"{BulkArgs(outputDir)} -maxsummarytablestringlength 5");
+
+			Assert.DoesNotContain("[ERROR]", output);
+			string html = ReadFirstHtml(outputDir);
+			Assert.NotNull(html);
+			// Truncated strings end with "..."
+			Assert.Contains("...", html);
+		}
+
+		#endregion
+
+		#region PerfLog
+
+		[Fact]
+		public void PerfLog_ProducesTimingOutput()
+		{
+			_fixture.RequireTestData();
+			string outputDir = Path.Combine(_outputDir, "PerfLogTest");
+
+			var (exitCode, output) = RunTool($"{BulkArgs(outputDir)} -perflog");
+
+			Assert.DoesNotContain("[ERROR]", output);
+			Assert.Contains("[PerfLog]", output);
+			// PerfLog lines include timing in ms format
+			Assert.Contains("ms", output);
+			Assert.Contains("TOTAL:", output);
+		}
+
+		#endregion
+
+		#region ListFiles
+
+		[Fact]
+		public void ListFiles_PrintsMatchingFilenames()
+		{
+			_fixture.RequireTestData();
+			string outputDir = Path.Combine(_outputDir, "ListFiles");
+
+			var (exitCode, output) = RunTool(
+				$"-csvdir \"{_fixture.GetPerfCsvDir()}\" -recurse -o \"{outputDir}\" " +
+				$"-searchpattern Perf_*.csv.bin -listfiles");
+
+			Assert.DoesNotContain("[ERROR]", output);
+			// Should print filenames that match the search pattern
+			Assert.Contains("Perf_", output);
+			Assert.Contains(".csv.bin", output);
+		}
+
+		#endregion
+
+		#region Regression Analysis
+
+		[Fact]
+		public void OnlyShowRegressedColumns_ProducesOutput()
+		{
+			_fixture.RequireTestData();
+			string outputDir = Path.Combine(_outputDir, "Regression");
+
+			var (exitCode, output) = RunTool($"{BulkArgs(outputDir)} -onlyshowregressedcolumns");
+
+			Assert.DoesNotContain("[ERROR]", output);
+			Assert.True(Directory.GetFiles(outputDir, "*.html").Length > 0);
+		}
+
+		[Fact]
+		public void RegressionWithCustomThresholds_ProducesOutput()
+		{
+			_fixture.RequireTestData();
+			string outputDir = Path.Combine(_outputDir, "RegressionCustom");
+
+			var (exitCode, output) = RunTool(
+				$"{BulkArgs(outputDir)} " +
+				$"-onlyshowregressedcolumns -regressionstddevthreshold 3 -regressionoutlierstddevthreshold 5");
+
+			Assert.DoesNotContain("[ERROR]", output);
+			Assert.True(Directory.GetFiles(outputDir, "*.html").Length > 0);
+		}
+
+		#endregion
+
+		#region Require Metadata
+
+		[Fact]
+		public void RequireMetadata_ProducesOutput()
+		{
+			_fixture.RequireTestData();
+			string outputDir = Path.Combine(_outputDir, "RequireMeta");
+
+			var (exitCode, output) = RunTool($"{BulkArgs(outputDir)} -requiremetadata");
+
+			Assert.DoesNotContain("[ERROR]", output);
+		}
+
+		#endregion
+
+		#region Condensed Table
+
+		[Fact]
+		public void TransposeCollatedTable_ContainsTransposeMarker()
+		{
+			_fixture.RequireTestData();
+			string outputDir = Path.Combine(_outputDir, "TransposeCollated");
+
+			var (exitCode, output) = RunTool($"{BulkArgs(outputDir)} -collatetable -transposecollatedtable");
+
+			Assert.DoesNotContain("[ERROR]", output);
+			// Check the collated file for transpose marker
+			string[] collatedFiles = Directory.GetFiles(outputDir, "*Collated*");
+			Assert.True(collatedFiles.Length > 0);
+			string html = File.ReadAllText(collatedFiles[0]);
+			Assert.Contains("SummaryTableIsTransposed", html);
+		}
+
+		#endregion
+
+		#region NoCsvCacheFiles
+
+		[Fact]
+		public void NoCsvCacheFiles_ProducesOutput()
+		{
+			_fixture.RequireTestData();
+			string outputDir = Path.Combine(_outputDir, "NoCsvCache");
+
+			var (exitCode, output) = RunTool($"{BulkArgs(outputDir)} -nocsvcachefiles");
+
+			Assert.DoesNotContain("[ERROR]", output);
+			Assert.True(Directory.GetFiles(outputDir, "*.html").Length > 0);
+		}
+
+		#endregion
+
+		#region MinFrameCount
+
+		[Fact]
+		public void MinFrameCount_ExcludesCsvsAndLogs()
+		{
+			_fixture.RequireTestData();
+			string outputDir = Path.Combine(_outputDir, "MinFrameCount");
+
+			// Use a very high frame count to filter out all CSVs
+			var (exitCode, output) = RunTool($"{BulkArgs(outputDir)} -minframecount 999999999");
+
+			Assert.DoesNotContain("[ERROR]", output);
+			// Should log that CSVs were excluded due to frame count
+			Assert.Contains("frame count", output.ToLower());
+		}
+
+		#endregion
+
+		#region Sort Trailing Digits As Numeric
+
+		[Fact]
+		public void SortTrailingDigitsAsNumeric_ProducesOutput()
+		{
+			_fixture.RequireTestData();
+			string outputDir = Path.Combine(_outputDir, "NumSort");
+
+			var (exitCode, output) = RunTool($"{BulkArgs(outputDir)} -sorttrailingdigitsasnumeric");
+
+			Assert.DoesNotContain("[ERROR]", output);
+			Assert.True(Directory.GetFiles(outputDir, "*.html").Length > 0);
+		}
+
+		#endregion
+	}
+}
diff --git a/Engine/Source/Programs/CSVTools/PerfReportTool.Tests/Integration/EndToEndTests.cs b/Engine/Source/Programs/CSVTools/PerfReportTool.Tests/Integration/EndToEndTests.cs
new file mode 100644
--- /dev/null
+++ b/Engine/Source/Programs/CSVTools/PerfReportTool.Tests/Integration/EndToEndTests.cs
@@ -0,0 +1,451 @@
+using System;
+using System.Collections.Generic;
+using System.Diagnostics;
+using System.IO;
+using System.Linq;
+using System.Text.Json;
+using Xunit;
+
+namespace PerfReportTool.Tests.Integration
+{
+	[Trait("Category", "Integration")]
+	[Collection("EndToEnd")]
+	public class EndToEndTests : IClassFixture<TestDataFixture>, IDisposable
+	{
+		private readonly TestDataFixture _fixture;
+		private readonly string _outputDir;
+		private readonly string _toolPath;
+
+		public EndToEndTests(TestDataFixture fixture)
+		{
+			_fixture = fixture;
+
+			// Create a temp output directory for this test run
+			_outputDir = Path.Combine(Path.GetTempPath(), "PerfReportToolTests_" + Path.GetRandomFileName());
+			Directory.CreateDirectory(_outputDir);
+
+			// Find the PerfReportTool executable
+			string assemblyDir = Path.GetDirectoryName(typeof(PerfReportTool.Version).Assembly.Location);
+			_toolPath = Path.Combine(assemblyDir, "PerfreportTool.dll");
+		}
+
+		public void Dispose()
+		{
+			if (Directory.Exists(_outputDir))
+			{
+				try { Directory.Delete(_outputDir, true); } catch { }
+			}
+		}
+
+		private (int exitCode, string output) RunPerfReportTool(string arguments, int timeoutMs = 120000)
+		{
+			var psi = new ProcessStartInfo
+			{
+				FileName = "dotnet",
+				Arguments = $"\"{_toolPath}\" {arguments}",
+				RedirectStandardOutput = true,
+				RedirectStandardError = true,
+				UseShellExecute = false,
+				CreateNoWindow = true
+			};
+
+			using var process = Process.Start(psi);
+			string stdout = process.StandardOutput.ReadToEnd();
+			string stderr = process.StandardError.ReadToEnd();
+			if (!process.WaitForExit(timeoutMs))
+			{
+				try { process.Kill(); } catch { }
+				return (-1, stdout + stderr + "\n[ERROR] Process timed out");
+			}
+
+			return (process.ExitCode, stdout + stderr);
+		}
+
+		/// <summary>
+		/// Helper: generate PRC cache from the Perf test CSVs and return the cache dir path.
+		/// </summary>
+		private string GeneratePerfPrcCache(string cacheName)
+		{
+			string cacheDir = Path.Combine(_outputDir, cacheName);
+			Directory.CreateDirectory(cacheDir);
+			string outputDir = Path.Combine(_outputDir, cacheName + "_html");
+
+			var (exitCode, output) = RunPerfReportTool(
+				$"-csvdir \"{_fixture.GetPerfCsvDir()}\" -emailtable -recurse -o \"{outputDir}\" " +
+				$"-nowatermarks -nodetailedreports " +
+				$"-searchpattern Perf_*.csv.bin " +
+				$"-summaryTableCacheUseOnlyCsvID -summaryTableCacheInvalidate " +
+				$"-summaryTableCache \"{cacheDir}\"",
+				timeoutMs: 300000);
+
+			Assert.DoesNotContain("[ERROR]", output);
+			Assert.True(Directory.GetFiles(cacheDir, "*.prc").Length > 0,
+				$"PRC cache generation failed — no .prc files in {cacheDir}. Output: {output.Substring(0, Math.Min(500, output.Length))}");
+
+			return cacheDir;
+		}
+
+		[Fact]
+		public void VersionOutput_PrintsVersion()
+		{
+			// Running with no args should print the version/help
+			var (exitCode, output) = RunPerfReportTool("");
+			Assert.Contains("PerfReportTool v", output);
+		}
+
+		[Fact]
+		public void BulkMode_CsvDir_ProducesSummaryTable()
+		{
+			_fixture.RequireTestData();
+
+			string outputDir = Path.Combine(_outputDir, "BulkCsvDir");
+			string cacheDir = Path.Combine(_outputDir, "Cache1");
+			Directory.CreateDirectory(cacheDir);
+
+			string csvDir = _fixture.GetPerfCsvDir();
+			var (exitCode, output) = RunPerfReportTool(
+				$"-csvdir \"{csvDir}\" -emailtable -recurse -o \"{outputDir}\" " +
+				$"-perflog -nowatermarks -nodetailedreports " +
+				$"-searchpattern *.csv.bin " +
+				$"-summaryTableCacheUseOnlyCsvID -summaryTableCacheInvalidate " +
+				$"-summaryTableCache \"{cacheDir}\"",
+				timeoutMs: 300000);
+
+			Assert.True(Directory.Exists(outputDir), "Output directory was not created");
+
+			string[] htmlFiles = Directory.GetFiles(outputDir, "*.html", SearchOption.TopDirectoryOnly);
+			Assert.True(htmlFiles.Length > 0, $"No HTML files generated. Output: {output.Substring(0, Math.Min(500, output.Length))}");
+
+			string[] prcFiles = Directory.GetFiles(cacheDir, "*.prc");
+			Assert.True(prcFiles.Length > 0, "No PRC cache files generated");
+
+			Assert.DoesNotContain("[ERROR]", output);
+		}
+
+		[Fact]
+		public void BulkMode_WithCache_ReusesCachedData()
+		{
+			_fixture.RequireTestData();
+
+			string cacheDir = Path.Combine(_outputDir, "Cache2");
+			Directory.CreateDirectory(cacheDir);
+			string csvDir = _fixture.GetPerfCsvDir();
+
+			string outputDir1 = Path.Combine(_outputDir, "WithCache_Pass1");
+			RunPerfReportTool(
+				$"-csvdir \"{csvDir}\" -emailtable -recurse -o \"{outputDir1}\" " +
+				$"-nowatermarks -nodetailedreports " +
+				$"-searchpattern *.csv.bin " +
+				$"-summaryTableCacheUseOnlyCsvID -summaryTableCacheInvalidate " +
+				$"-summaryTableCache \"{cacheDir}\"",
+				timeoutMs: 300000);
+
+			int prcCount = Directory.GetFiles(cacheDir, "*.prc").Length;
+			Assert.True(prcCount > 0, "First pass should have created PRC files");
+
+			string outputDir2 = Path.Combine(_outputDir, "WithCache_Pass2");
+			var (exitCode2, output2) = RunPerfReportTool(
+				$"-csvdir \"{csvDir}\" -emailtable -recurse -o \"{outputDir2}\" " +
+				$"-nowatermarks -nodetailedreports " +
+				$"-searchpattern *.csv.bin " +
+				$"-summaryTableCacheUseOnlyCsvID " +
+				$"-summaryTableCache \"{cacheDir}\"",
+				timeoutMs: 300000);
+
+			Assert.True(Directory.Exists(outputDir2), "Second pass output directory not created");
+			Assert.DoesNotContain("[ERROR]", output2);
+		}
+
+		[Fact]
+		public void BulkMode_MetadataFilter_FiltersCorrectly()
+		{
+			_fixture.RequireTestData();
+
+			string outputDir = Path.Combine(_outputDir, "MetadataFilter");
+			string cacheDir = Path.Combine(_outputDir, "Cache3");
+			Directory.CreateDirectory(cacheDir);
+
+			string csvDir = _fixture.GetPerfCsvDir();
+			string filter = "gauntletsubtest=Perf";
+
+			var (exitCode, output) = RunPerfReportTool(
+				$"-csvdir \"{csvDir}\" -emailtable -recurse -o \"{outputDir}\" " +
+				$"-perflog -nodetailedreports -nowatermarks " +
+				$"-metadatafilter \"{filter}\" " +
+				$"-searchpattern *.csv.bin " +
+				$"-summaryTableCacheUseOnlyCsvID -summaryTableCache \"{cacheDir}\"",
+				timeoutMs: 300000);
+
+			Assert.True(Directory.Exists(outputDir), "Output directory not created for metadata filter test");
+			Assert.DoesNotContain("[ERROR]", output);
+		}
+
+		[Fact]
+		public void BulkMode_LlmReport_ProducesOutput()
+		{
+			_fixture.RequireTestData();
+
+			string outputDir = Path.Combine(_outputDir, "LLM");
+			string csvDir = _fixture.GetPerfCsvDir();
+
+			var (exitCode, output) = RunPerfReportTool(
+				$"-csvdir \"{csvDir}\" -o \"{outputDir}\" -nowatermarks " +
+				$"-reportxml LLMReportTypes.xml -reportType LLM " +
+				$"-searchpattern LLM_*.csv.bin",
+				timeoutMs: 120000);
+
+			Assert.DoesNotContain("[ERROR]", output);
+			Assert.True(Directory.Exists(outputDir), "LLM output directory not created");
+		}
+
+		#region JSON serialization tests
+
+		[Fact]
+		public void SummaryTableToJson_WriteAllElementData_ContainsAllTypes()
+		{
+			_fixture.RequireTestData();
+
+			string cacheDir = GeneratePerfPrcCache("JsonAllTypes");
+
+			string jsonPath = Path.Combine(_outputDir, "AllElementData.json");
+			var (exitCode, output) = RunPerfReportTool(
+				$"-summarytablecachein \"{cacheDir}\" " +
+				$"-summaryTableToJsonWriteAllElementData " +
+				$"-summarytabletojsonfastmode " +
+				$"-summaryTableToJson \"{jsonPath}\" " +
+				$"-precacheThreads 1",
+				timeoutMs: 120000);
+
+			Assert.DoesNotContain("[ERROR]", output);
+			Assert.True(File.Exists(jsonPath), "JSON output not generated");
+
+			// Parse and validate structure
+			string jsonContent = File.ReadAllText(jsonPath);
+			var root = JsonSerializer.Deserialize<Dictionary<string, JsonElement>>(jsonContent);
+			Assert.NotNull(root);
+			Assert.True(root.Count > 0, "JSON should contain at least one CSV entry");
+
+			// Each CSV entry should have keys for each SummaryTableElement.Type (except COUNT and ExternalMetadata)
+			string[] expectedTypeKeys = { "CsvStatAverage", "CsvMetadata", "SummaryTableMetric", "ToolMetadata", "CsvStatMin", "CsvStatMax" };
+
+			foreach (var csvEntry in root)
+			{
+				var entryDict = JsonSerializer.Deserialize<Dictionary<string, JsonElement>>(csvEntry.Value.GetRawText());
+				foreach (string typeKey in expectedTypeKeys)
+				{
+					Assert.True(entryDict.ContainsKey(typeKey),
+						$"CSV entry '{csvEntry.Key}' missing expected type key '{typeKey}'");
+				}
+
+				// CsvMetadata should have values like platform, buildversion
+				var metadata = JsonSerializer.Deserialize<Dictionary<string, JsonElement>>(entryDict["CsvMetadata"].GetRawText());
+				Assert.True(metadata.Count > 0, "CsvMetadata should have entries");
+
+				// With WriteAllElementData, each element should be an object with a "value" key
+				foreach (var metaEntry in metadata)
+				{
+					Assert.Equal(JsonValueKind.Object, metaEntry.Value.ValueKind);
+					var elementDict = JsonSerializer.Deserialize<Dictionary<string, JsonElement>>(metaEntry.Value.GetRawText());
+					Assert.True(elementDict.ContainsKey("value"), $"Element '{metaEntry.Key}' missing 'value' key");
+				}
+
+				// CsvStatAverage should have stat entries
+				var statAverages = JsonSerializer.Deserialize<Dictionary<string, JsonElement>>(entryDict["CsvStatAverage"].GetRawText());
+				Assert.True(statAverages.Count > 0, "CsvStatAverage should have stat entries");
+			}
+		}
+
+		[Fact]
+		public void SummaryTableToJson_MetadataOnly_OnlyWritesCsvMetadata()
+		{
+			_fixture.RequireTestData();
+
+			string cacheDir = GeneratePerfPrcCache("JsonMetaOnly");
+
+			string jsonPath = Path.Combine(_outputDir, "MetadataOnly.json");
+			var (exitCode, output) = RunPerfReportTool(
+				$"-summarytablecachein \"{cacheDir}\" " +
+				$"-summaryTableToJsonMetadataOnly " +
+				$"-summarytabletojsonfastmode " +
+				$"-summaryTableToJson \"{jsonPath}\" " +
+				$"-precacheThreads 1",
+				timeoutMs: 120000);
+
+			Assert.DoesNotContain("[ERROR]", output);
+			Assert.True(File.Exists(jsonPath), "JSON output not generated");
+
+			string jsonContent = File.ReadAllText(jsonPath);
+			var root = JsonSerializer.Deserialize<Dictionary<string, JsonElement>>(jsonContent);
+			Assert.True(root.Count > 0, "JSON should contain at least one CSV entry");
+
+			foreach (var csvEntry in root)
+			{
+				var entryDict = JsonSerializer.Deserialize<Dictionary<string, JsonElement>>(csvEntry.Value.GetRawText());
+
+				// Should only contain CsvMetadata
+				Assert.True(entryDict.ContainsKey("CsvMetadata"), "Missing CsvMetadata key");
+				Assert.Single(entryDict); // Only CsvMetadata, nothing else
+
+				// Should not contain stat types
+				Assert.False(entryDict.ContainsKey("CsvStatAverage"), "MetadataOnly should not include CsvStatAverage");
+				Assert.False(entryDict.ContainsKey("ToolMetadata"), "MetadataOnly should not include ToolMetadata");
+			}
+		}
+
+		[Fact]
+		public void SummaryTableToJson_SeparateFiles_ProducesMultipleFiles()
+		{
+			_fixture.RequireTestData();
+
+			string cacheDir = GeneratePerfPrcCache("JsonSeparate");
+
+			string jsonDir = Path.Combine(_outputDir, "SeparateJsonFiles");
+			Directory.CreateDirectory(jsonDir);
+			var (exitCode, output) = RunPerfReportTool(
+				$"-summarytablecachein \"{cacheDir}\" " +
+				$"-summaryTableToJsonSeparateFiles " +
+				$"-summarytabletojsonfastmode " +
+				$"-summaryTableToJson \"{jsonDir}\" " +
+				$"-precacheThreads 1",
+				timeoutMs: 120000);
+
+			Assert.DoesNotContain("[ERROR]", output);
+
+			// Should produce one .json file per CSV (we have 2 Perf CSVs)
+			string[] jsonFiles = Directory.GetFiles(jsonDir, "*.json");
+			Assert.True(jsonFiles.Length >= 2, $"Expected at least 2 separate JSON files, got {jsonFiles.Length}");
+
+			// Each file should be valid JSON with the per-CSV structure
+			foreach (string jsonFile in jsonFiles)
+			{
+				string content = File.ReadAllText(jsonFile);
+				var entryDict = JsonSerializer.Deserialize<Dictionary<string, JsonElement>>(content);
+				Assert.NotNull(entryDict);
+				Assert.True(entryDict.ContainsKey("CsvMetadata"), $"File {Path.GetFileName(jsonFile)} missing CsvMetadata");
+			}
+		}
+
+		[Fact]
+		public void SummaryTableToJson_DefaultMode_WritesValuesDirectly()
+		{
+			_fixture.RequireTestData();
+
+			string cacheDir = GeneratePerfPrcCache("JsonDefault");
+
+			// Without WriteAllElementData, values should be written directly (not as objects with "value" key)
+			string jsonPath = Path.Combine(_outputDir, "Default.json");
+			var (exitCode, output) = RunPerfReportTool(
+				$"-summarytablecachein \"{cacheDir}\" " +
+				$"-summarytabletojsonfastmode " +
+				$"-summaryTableToJson \"{jsonPath}\" " +
+				$"-precacheThreads 1",
+				timeoutMs: 120000);
+
+			Assert.DoesNotContain("[ERROR]", output);
+			Assert.True(File.Exists(jsonPath), "JSON output not generated");
+
+			string jsonContent = File.ReadAllText(jsonPath);
+			var root = JsonSerializer.Deserialize<Dictionary<string, JsonElement>>(jsonContent);
+			Assert.NotNull(root);
+			Assert.True(root.Count > 0, "JSON should contain at least one CSV entry");
+
+			foreach (var csvEntry in root)
+			{
+				var entryDict = JsonSerializer.Deserialize<Dictionary<string, JsonElement>>(csvEntry.Value.GetRawText());
+				var metadata = JsonSerializer.Deserialize<Dictionary<string, JsonElement>>(entryDict["CsvMetadata"].GetRawText());
+
+				// Without WriteAllElementData, metadata values should be direct strings/numbers, not objects
+				foreach (var metaEntry in metadata)
+				{
+					Assert.NotEqual(JsonValueKind.Object, metaEntry.Value.ValueKind);
+				}
+			}
+		}
+
+		[Fact]
+		public void SummaryTableToJson_NoIndent_ProducesCompactJson()
+		{
+			_fixture.RequireTestData();
+
+			string cacheDir = GeneratePerfPrcCache("JsonNoIndent");
+
+			string jsonPath = Path.Combine(_outputDir, "NoIndent.json");
+			var (exitCode, output) = RunPerfReportTool(
+				$"-summarytablecachein \"{cacheDir}\" " +
+				$"-summaryTableToJsonNoIndent " +
+				$"-summarytabletojsonfastmode " +
+				$"-summaryTableToJson \"{jsonPath}\" " +
+				$"-precacheThreads 1",
+				timeoutMs: 120000);
+
+			Assert.DoesNotContain("[ERROR]", output);
+			Assert.True(File.Exists(jsonPath), "JSON output not generated");
+
+			string jsonContent = File.ReadAllText(jsonPath);
+			// Compact JSON should have no newlines in the content (single line)
+			int lineCount = jsonContent.Split('\n').Length;
+			Assert.True(lineCount <= 2, $"NoIndent JSON should be compact but has {lineCount} lines");
+		}
+
+		[Fact]
+		public void SummaryTableToJson_FileStream_ProducesValidJson()
+		{
+			_fixture.RequireTestData();
+
+			string cacheDir = GeneratePerfPrcCache("JsonStream");
+
+			string jsonPath = Path.Combine(_outputDir, "Streamed.json");
+			var (exitCode, output) = RunPerfReportTool(
+				$"-summarytablecachein \"{cacheDir}\" " +
+				$"-summaryTableToJsonFileStream " +
+				$"-summarytabletojsonfastmode " +
+				$"-summaryTableToJson \"{jsonPath}\" " +
+				$"-precacheThreads 1",
+				timeoutMs: 120000);
+
+			Assert.DoesNotContain("[ERROR]", output);
+			Assert.True(File.Exists(jsonPath), "JSON output not generated");
+
+			// Verify it's parseable JSON
+			string jsonContent = File.ReadAllText(jsonPath);
+			var root = JsonSerializer.Deserialize<Dictionary<string, JsonElement>>(jsonContent);
+			Assert.NotNull(root);
+			Assert.True(root.Count > 0, "Streamed JSON should contain CSV entries");
+		}
+
+		[Fact]
+		public void JsonToPrcs_ConvertsJsonBackToPrcFiles()
+		{
+			_fixture.RequireTestData();
+
+			string cacheDir = GeneratePerfPrcCache("JsonToPrc_Src");
+
+			// Export to JSON
+			string jsonPath = Path.Combine(_outputDir, "ForPrcConvert.json");
+			RunPerfReportTool(
+				$"-summarytablecachein \"{cacheDir}\" " +
+				$"-summaryTableToJsonWriteAllElementData " +
+				$"-summarytabletojsonfastmode " +
+				$"-summaryTableToJson \"{jsonPath}\" " +
+				$"-precacheThreads 1",
+				timeoutMs: 120000);
+
+			Assert.True(File.Exists(jsonPath), "Source JSON not generated");
+
+			// Convert JSON back to PRCs
+			string prcOutputDir = Path.Combine(_outputDir, "ConvertedPrcs");
+			Directory.CreateDirectory(prcOutputDir);
+			var (exitCode, output) = RunPerfReportTool(
+				$"-jsonToPrcs \"{jsonPath}\" " +
+				$"-summaryTableCache \"{prcOutputDir}\"",
+				timeoutMs: 120000);
+
+			Assert.DoesNotContain("[ERROR]", output);
+
+			string[] prcFiles = Directory.GetFiles(prcOutputDir, "*.prc");
+			Assert.True(prcFiles.Length >= 2, $"Expected at least 2 PRC files from JSON conversion, got {prcFiles.Length}");
+		}
+
+		#endregion
+	}
+}
diff --git a/Engine/Source/Programs/CSVTools/PerfReportTool.Tests/PerfReportTool.Tests.csproj b/Engine/Source/Programs/CSVTools/PerfReportTool.Tests/PerfReportTool.Tests.csproj
new file mode 100644
--- /dev/null
+++ b/Engine/Source/Programs/CSVTools/PerfReportTool.Tests/PerfReportTool.Tests.csproj
@@ -0,0 +1,19 @@
+<Project Sdk="Microsoft.NET.Sdk">
+  <PropertyGroup>
+    <TargetFramework>net10.0</TargetFramework>
+    <IsPackable>false</IsPackable>
+    <IsTestProject>true</IsTestProject>
+  </PropertyGroup>
+  <ItemGroup>
+    <PackageReference Include="Microsoft.NET.Test.Sdk" Version="17.14.1" />
+    <PackageReference Include="xunit" Version="2.9.3" />
+    <PackageReference Include="xunit.runner.visualstudio" Version="2.8.2" />
+  </ItemGroup>
+  <ItemGroup>
+    <Content Include="TestData\**" CopyToOutputDirectory="PreserveNewest" />
+  </ItemGroup>
+  <ItemGroup>
+    <ProjectReference Include="..\PerfReportTool\PerfreportTool.csproj" />
+    <ProjectReference Include="..\CsvStats\CsvStats.csproj" />
+  </ItemGroup>
+</Project>
diff --git a/Engine/Source/Programs/CSVTools/PerfReportTool.Tests/PerfReportToolTests/ColourThresholdListTests.cs b/Engine/Source/Programs/CSVTools/PerfReportTool.Tests/PerfReportToolTests/ColourThresholdListTests.cs
new file mode 100644
--- /dev/null
+++ b/Engine/Source/Programs/CSVTools/PerfReportTool.Tests/PerfReportToolTests/ColourThresholdListTests.cs
@@ -0,0 +1,219 @@
+using System.Collections.Generic;
+using PerfSummaries;
+using Xunit;
+
+namespace PerfReportTool.Tests.PerfReportToolTests
+{
+	public class ColourTests
+	{
+		[Fact]
+		public void Colour_FromHexString_ParsesCorrectly()
+		{
+			var colour = new Colour("#ff0000");
+			Assert.Equal(1.0f, colour.r, 0.01f);
+			Assert.Equal(0.0f, colour.g, 0.01f);
+			Assert.Equal(0.0f, colour.b, 0.01f);
+		}
+
+		[Fact]
+		public void Colour_FromHexString_NoHash_ParsesCorrectly()
+		{
+			var colour = new Colour("00ff00");
+			Assert.Equal(0.0f, colour.r, 0.01f);
+			Assert.Equal(1.0f, colour.g, 0.01f);
+			Assert.Equal(0.0f, colour.b, 0.01f);
+		}
+
+		[Fact]
+		public void Colour_ToString_ProducesHexString()
+		{
+			var colour = new Colour(1.0f, 0.0f, 0.0f);
+			string result = colour.ToString();
+			Assert.Equal("#ff0000", result);
+		}
+
+		[Fact]
+		public void Colour_ToString_RoundTrips()
+		{
+			var original = new Colour("#3a7fc2");
+			string str = original.ToString();
+			var restored = new Colour(str);
+			Assert.Equal(original.r, restored.r, 0.01f);
+			Assert.Equal(original.g, restored.g, 0.01f);
+			Assert.Equal(original.b, restored.b, 0.01f);
+		}
+
+		[Fact]
+		public void Colour_Lerp_MidpointIsCorrect()
+		{
+			var black = new Colour(0.0f, 0.0f, 0.0f);
+			var white = new Colour(1.0f, 1.0f, 1.0f);
+			var mid = Colour.Lerp(black, white, 0.5f);
+			Assert.Equal(0.5f, mid.r, 0.01f);
+			Assert.Equal(0.5f, mid.g, 0.01f);
+			Assert.Equal(0.5f, mid.b, 0.01f);
+		}
+
+		[Fact]
+		public void Colour_Lerp_ZeroReturnFirst()
+		{
+			var a = new Colour(1.0f, 0.0f, 0.0f);
+			var b = new Colour(0.0f, 1.0f, 0.0f);
+			var result = Colour.Lerp(a, b, 0.0f);
+			Assert.Equal(a.r, result.r, 0.001f);
+			Assert.Equal(a.g, result.g, 0.001f);
+			Assert.Equal(a.b, result.b, 0.001f);
+		}
+
+		[Fact]
+		public void Colour_Lerp_OneReturnSecond()
+		{
+			var a = new Colour(1.0f, 0.0f, 0.0f);
+			var b = new Colour(0.0f, 1.0f, 0.0f);
+			var result = Colour.Lerp(a, b, 1.0f);
+			Assert.Equal(b.r, result.r, 0.001f);
+			Assert.Equal(b.g, result.g, 0.001f);
+			Assert.Equal(b.b, result.b, 0.001f);
+		}
+	}
+
+	public class ColourThresholdListTests
+	{
+		[Fact]
+		public void GetColourForValue_BelowLowest_ReturnsGreen()
+		{
+			// Constructor params: (red=40, orange=30, yellow=20, green=10)
+			// Internally stored in ascending order: green(10), yellow(20), orange(30), red(40)
+			var list = new ColourThresholdList(40, 30, 20, 10);
+			string colour = list.GetColourForValue(5.0);
+
+			// Below the green threshold should return the green colour
+			string expectedGreen = Colour.Green.ToHTMLString();
+			Assert.Equal(expectedGreen, colour);
+		}
+
+		[Fact]
+		public void GetColourForValue_AboveHighest_ReturnsRed()
+		{
+			var list = new ColourThresholdList(40, 30, 20, 10);
+			string colour = list.GetColourForValue(50.0);
+
+			string expectedRed = Colour.Red.ToHTMLString();
+			Assert.Equal(expectedRed, colour);
+		}
+
+		[Fact]
+		public void GetColourForValue_ExactThreshold_ReturnsCorrectColour()
+		{
+			// With lerpColours=false, exact threshold should return that threshold's colour
+			var list = new ColourThresholdList(40, 30, 20, 10, lerpColours: false);
+			string colourAtGreen = list.GetColourForValue(10.0);
+			Assert.Equal(Colour.Green.ToHTMLString(), colourAtGreen);
+		}
+
+		[Fact]
+		public void GetColourForValue_MidValue_LerpsColours()
+		{
+			var list = new ColourThresholdList(40, 30, 20, 10);
+			string colourAt15 = list.GetColourForValue(15.0); // Between green(10) and yellow(20)
+
+			// Should not be pure green or pure yellow when lerping
+			Assert.NotEqual(Colour.Green.ToHTMLString(), colourAt15);
+			Assert.NotEqual(Colour.Yellow.ToHTMLString(), colourAt15);
+			Assert.NotEqual("'#ffffff'", colourAt15);
+		}
+
+		[Fact]
+		public void GetColourForValue_NoLerp_SnapsToThreshold()
+		{
+			var list = new ColourThresholdList(40, 30, 20, 10, lerpColours: false);
+
+			// Without lerping, values between thresholds snap to the threshold they fall within
+			string colourAt15 = list.GetColourForValue(15.0); // Between green(10) and yellow(20)
+			string colourAt25 = list.GetColourForValue(25.0); // Between yellow(20) and orange(30)
+
+			// These two should be different colours (different thresholds)
+			Assert.NotEqual(colourAt15, colourAt25);
+			// Neither should be the lerped result or white
+			Assert.NotEqual("'#ffffff'", colourAt15);
+			Assert.NotEqual("'#ffffff'", colourAt25);
+		}
+
+		[Fact]
+		public void GetColourForValue_EmptyList_ReturnsWhite()
+		{
+			var list = new ColourThresholdList();
+			string colour = list.GetColourForValue(10.0);
+			Assert.Equal("'#ffffff'", colour);
+		}
+
+		[Fact]
+		public void GetColourForValue_StringValue_ConvertsAndReturns()
+		{
+			var list = new ColourThresholdList(40, 30, 20, 10);
+			string colourFromString = list.GetColourForValue("25.0");
+			string colourFromDouble = list.GetColourForValue(25.0);
+
+			Assert.Equal(colourFromDouble, colourFromString);
+		}
+
+		[Fact]
+		public void GetColourForValue_NonNumericString_ReturnsWhite()
+		{
+			var list = new ColourThresholdList(40, 30, 20, 10);
+			string colour = list.GetColourForValue("not_a_number");
+			Assert.Equal("'#ffffff'", colour);
+		}
+
+		[Fact]
+		public void GetSafeColourForValue_NullList_ReturnsWhite()
+		{
+			Assert.Equal("'#ffffff'", ColourThresholdList.GetSafeColourForValue(null, 10.0));
+			Assert.Equal("'#ffffff'", ColourThresholdList.GetSafeColourForValue(null, "10.0"));
+		}
+
+		[Fact]
+		public void ColourThresholdList_Count_MatchesAddedThresholds()
+		{
+			var list = new ColourThresholdList();
+			Assert.Equal(0, list.Count);
+
+			list.Add(new ThresholdInfo(10.0, Colour.Green));
+			list.Add(new ThresholdInfo(20.0, Colour.Red));
+			Assert.Equal(2, list.Count);
+		}
+
+		[Fact]
+		public void ColourThresholdList_ToJsonDict_ContainsThresholdValues()
+		{
+			var list = new ColourThresholdList(40, 30, 20, 10);
+			var jsonDict = list.ToJsonDict();
+
+			Assert.True(jsonDict.ContainsKey("thresholds"));
+			Assert.True(jsonDict.ContainsKey("lerpColours"));
+
+			var thresholds = (List<double>)jsonDict["thresholds"];
+			Assert.Equal(4, thresholds.Count);
+			// Constructor adds internally in ascending order: green(10), yellow(20), orange(30), red(40)
+			Assert.Equal(10.0, thresholds[0]);
+			Assert.Equal(20.0, thresholds[1]);
+			Assert.Equal(30.0, thresholds[2]);
+			Assert.Equal(40.0, thresholds[3]);
+		}
+
+		[Fact]
+		public void GetThresholdColour_CustomColours_ReturnsCorrectColour()
+		{
+			var custom = new Colour("#112233");
+			var thresholds = new List<ThresholdInfo>
+			{
+				new ThresholdInfo(10.0, custom),
+				new ThresholdInfo(20.0, Colour.Red)
+			};
+			var list = new ColourThresholdList(thresholds, lerpColours: false);
+
+			string colourBelow = list.GetColourForValue(5.0);
+			Assert.Equal(custom.ToHTMLString(), colourBelow);
+		}
+	}
+}
diff --git a/Engine/Source/Programs/CSVTools/PerfReportTool.Tests/PerfReportToolTests/ColumnFilterTests.cs b/Engine/Source/Programs/CSVTools/PerfReportTool.Tests/PerfReportToolTests/ColumnFilterTests.cs
new file mode 100644
--- /dev/null
+++ b/Engine/Source/Programs/CSVTools/PerfReportTool.Tests/PerfReportToolTests/ColumnFilterTests.cs
@@ -0,0 +1,263 @@
+// Copyright Epic Games, Inc. All Rights Reserved.
+
+using System;
+using System.Collections.Generic;
+using PerfReportTool;
+using PerfSummaries;
+using Xunit;
+
+namespace PerfReportTool.Tests.PerfReportToolTests
+{
+	public class StatThresholdColumnFilterTests
+	{
+		private static SummaryTableColumn CreateNumericColumn(string name, SummaryTableElement.Type type, double[] values)
+		{
+			var col = new SummaryTableColumn(name, true, name, false, type);
+			for (int i = 0; i < values.Length; i++)
+			{
+				col.SetValue(i, values[i]);
+			}
+			return col;
+		}
+
+		private static SummaryTableColumn CreateStringColumn(string name, SummaryTableElement.Type type, string[] values)
+		{
+			var col = new SummaryTableColumn(name, false, name, false, type);
+			for (int i = 0; i < values.Length; i++)
+			{
+				col.SetStringValue(i, values[i]);
+			}
+			return col;
+		}
+
+		[Fact]
+		public void ShouldFilter_AllValuesAboveThreshold_ReturnsFalse()
+		{
+			var filter = new StatThresholdColumnFilter(10.0f);
+			var col = CreateNumericColumn("stat", SummaryTableElement.Type.CsvStatAverage, new[] { 20.0, 30.0, 40.0 });
+			var table = new SummaryTable();
+			Assert.False(filter.ShouldFilter(col, table));
+		}
+
+		[Fact]
+		public void ShouldFilter_SomeValuesBelowThreshold_ReturnsTrue()
+		{
+			var filter = new StatThresholdColumnFilter(10.0f);
+			var col = CreateNumericColumn("stat", SummaryTableElement.Type.CsvStatAverage, new[] { 5.0, 30.0, 40.0 });
+			var table = new SummaryTable();
+			Assert.True(filter.ShouldFilter(col, table));
+		}
+
+		[Fact]
+		public void ShouldFilter_AllValuesBelowThreshold_ReturnsTrue()
+		{
+			var filter = new StatThresholdColumnFilter(10.0f);
+			var col = CreateNumericColumn("stat", SummaryTableElement.Type.CsvStatAverage, new[] { 5.0, 3.0, 1.0 });
+			var table = new SummaryTable();
+			Assert.True(filter.ShouldFilter(col, table));
+		}
+
+		[Fact]
+		public void ShouldFilter_ThresholdZero_NeverFilters()
+		{
+			var filter = new StatThresholdColumnFilter(0.0f);
+			var col = CreateNumericColumn("stat", SummaryTableElement.Type.CsvStatAverage, new[] { -5.0, 0.0 });
+			var table = new SummaryTable();
+			Assert.False(filter.ShouldFilter(col, table));
+		}
+
+		[Fact]
+		public void ShouldFilter_NegativeThreshold_NeverFilters()
+		{
+			var filter = new StatThresholdColumnFilter(-1.0f);
+			var col = CreateNumericColumn("stat", SummaryTableElement.Type.CsvStatAverage, new[] { -5.0 });
+			var table = new SummaryTable();
+			Assert.False(filter.ShouldFilter(col, table));
+		}
+
+		[Fact]
+		public void ShouldFilter_NonNumericColumn_ReturnsFalse()
+		{
+			var filter = new StatThresholdColumnFilter(10.0f);
+			var col = CreateStringColumn("meta", SummaryTableElement.Type.CsvMetadata, new[] { "hello" });
+			var table = new SummaryTable();
+			Assert.False(filter.ShouldFilter(col, table));
+		}
+
+		[Fact]
+		public void ShouldFilter_MetadataType_ReturnsFalse()
+		{
+			// Even numeric metadata columns should not be filtered by threshold
+			var filter = new StatThresholdColumnFilter(10.0f);
+			var col = CreateNumericColumn("meta", SummaryTableElement.Type.CsvMetadata, new[] { 1.0 });
+			var table = new SummaryTable();
+			Assert.False(filter.ShouldFilter(col, table));
+		}
+
+		[Fact]
+		public void ShouldFilter_AllFilterableTypes_FiltersWhenBelowThreshold()
+		{
+			var typesToTest = new[]
+			{
+				SummaryTableElement.Type.CsvStatAverage,
+				SummaryTableElement.Type.CsvStatMin,
+				SummaryTableElement.Type.CsvStatMax,
+				SummaryTableElement.Type.SummaryTableMetric,
+				SummaryTableElement.Type.ExternalMetadata
+			};
+
+			foreach (var type in typesToTest)
+			{
+				var filter = new StatThresholdColumnFilter(100.0f);
+				var col = CreateNumericColumn("stat", type, new[] { 5.0 });
+				var table = new SummaryTable();
+				Assert.True(filter.ShouldFilter(col, table), $"Type {type} should be filtered when below threshold");
+			}
+		}
+	}
+
+	public class MetadataColumnFilterTests
+	{
+		private static SummaryTableColumn CreateColumn(string name, SummaryTableElement.Type type)
+		{
+			return new SummaryTableColumn(name, false, name, false, type);
+		}
+
+		[Fact]
+		public void ShouldFilter_CsvMetadataColumn_ReturnsTrue()
+		{
+			var filter = new MetadataColumnFilter(new List<string>());
+			var col = CreateColumn("platform", SummaryTableElement.Type.CsvMetadata);
+			Assert.True(filter.ShouldFilter(col, new SummaryTable()));
+		}
+
+		[Fact]
+		public void ShouldFilter_ToolMetadataColumn_ReturnsTrue()
+		{
+			var filter = new MetadataColumnFilter(new List<string>());
+			var col = CreateColumn("toolversion", SummaryTableElement.Type.ToolMetadata);
+			Assert.True(filter.ShouldFilter(col, new SummaryTable()));
+		}
+
+		[Fact]
+		public void ShouldFilter_CsvStatColumn_ReturnsFalse()
+		{
+			var filter = new MetadataColumnFilter(new List<string>());
+			var col = CreateColumn("frametime", SummaryTableElement.Type.CsvStatAverage);
+			Assert.False(filter.ShouldFilter(col, new SummaryTable()));
+		}
+
+		[Fact]
+		public void ShouldFilter_MetadataInIgnoreList_ReturnsFalse()
+		{
+			var filter = new MetadataColumnFilter(new List<string> { "platform" });
+			var col = CreateColumn("platform", SummaryTableElement.Type.CsvMetadata);
+			Assert.False(filter.ShouldFilter(col, new SummaryTable()));
+		}
+
+		[Fact]
+		public void ShouldFilter_MetadataNotInIgnoreList_ReturnsTrue()
+		{
+			var filter = new MetadataColumnFilter(new List<string> { "platform" });
+			var col = CreateColumn("buildversion", SummaryTableElement.Type.CsvMetadata);
+			Assert.True(filter.ShouldFilter(col, new SummaryTable()));
+		}
+
+		[Fact]
+		public void ShouldFilter_SummaryMetric_ReturnsFalse()
+		{
+			var filter = new MetadataColumnFilter(new List<string>());
+			var col = CreateColumn("MVP", SummaryTableElement.Type.SummaryTableMetric);
+			Assert.False(filter.ShouldFilter(col, new SummaryTable()));
+		}
+	}
+
+	public class RegressionColumnFilterTests
+	{
+		private static SummaryTableColumn CreateNumericColumn(string name, double[] values, SummaryTableElement.Type type = SummaryTableElement.Type.CsvStatAverage)
+		{
+			var col = new SummaryTableColumn(name, true, name, false, type);
+			for (int i = 0; i < values.Length; i++)
+			{
+				col.SetValue(i, values[i]);
+			}
+			return col;
+		}
+
+		private static SummaryTable CreateTableWithColumn(SummaryTableColumn column)
+		{
+			var table = new SummaryTable();
+			var rowData = new SummaryTableRowData();
+			// We need to add rows to the table to set its Count property.
+			// Add multiple rows by adding row data entries that match our column values.
+			for (int i = 0; i < column.GetCount(); i++)
+			{
+				var row = new SummaryTableRowData();
+				row.Add(SummaryTableElement.Type.CsvStatAverage, column.name, column.GetValue(i));
+				table.AddRowData(row, true, false, false, false);
+			}
+			return table;
+		}
+
+		[Fact]
+		public void ShouldFilter_NonNumericColumn_ReturnsFalse()
+		{
+			var filter = new RegressionColumnFilter(null, 2.0f, 4.0f);
+			var col = new SummaryTableColumn("meta", false, "meta", false, SummaryTableElement.Type.CsvMetadata);
+			col.SetStringValue(0, "test");
+			var table = CreateTableWithColumn(CreateNumericColumn("dummy", new[] { 1.0 }));
+			Assert.False(filter.ShouldFilter(col, table));
+		}
+
+		[Fact]
+		public void ShouldFilter_EmptyTable_ReturnsTrue()
+		{
+			var filter = new RegressionColumnFilter(null, 2.0f, 4.0f);
+			var col = CreateNumericColumn("stat", new double[0]);
+			var table = new SummaryTable();
+			Assert.True(filter.ShouldFilter(col, table));
+		}
+
+		[Fact]
+		public void ShouldFilter_ClearRegression_ReturnsFalse()
+		{
+			// First row group (1 row) has a significantly higher value than historical data
+			// Historical: ~10.0, First row: 50.0 => clear regression for HighIsBad
+			var filter = new RegressionColumnFilter(null, 2.0f, 4.0f);
+			var col = CreateNumericColumn("stat", new[] { 50.0, 10.0, 10.5, 9.5, 10.0, 11.0, 9.0 });
+			var table = CreateTableWithColumn(col);
+			Assert.False(filter.ShouldFilter(col, table));
+		}
+
+		[Fact]
+		public void ShouldFilter_NoRegression_ReturnsTrue()
+		{
+			// All values are similar, no regression
+			var filter = new RegressionColumnFilter(null, 2.0f, 4.0f);
+			var col = CreateNumericColumn("stat", new[] { 10.0, 10.1, 9.9, 10.0, 10.2, 9.8 });
+			var table = CreateTableWithColumn(col);
+			Assert.True(filter.ShouldFilter(col, table));
+		}
+
+		[Fact]
+		public void ShouldFilter_TooManyMissingValues_ReturnsTrue()
+		{
+			// >60% missing values should filter out
+			var filter = new RegressionColumnFilter(null, 2.0f, 4.0f);
+			// Create column with mostly MaxValue (missing) entries
+			var col = CreateNumericColumn("stat", new[] { 50.0, double.MaxValue, double.MaxValue, double.MaxValue, double.MaxValue, double.MaxValue, double.MaxValue, double.MaxValue, 10.0, double.MaxValue });
+			var table = CreateTableWithColumn(col);
+			Assert.True(filter.ShouldFilter(col, table));
+		}
+
+		[Fact]
+		public void ShouldFilter_MetadataTypeColumn_ReturnsFalse()
+		{
+			// Metadata columns should pass through (not filtered by regression)
+			var filter = new RegressionColumnFilter(null, 2.0f, 4.0f);
+			var col = CreateNumericColumn("meta", new[] { 50.0, 10.0 }, SummaryTableElement.Type.CsvMetadata);
+			var table = CreateTableWithColumn(CreateNumericColumn("dummy", new[] { 50.0, 10.0 }));
+			Assert.False(filter.ShouldFilter(col, table));
+		}
+	}
+}
diff --git a/Engine/Source/Programs/CSVTools/PerfReportTool.Tests/PerfReportToolTests/CommandLineTests.cs b/Engine/Source/Programs/CSVTools/PerfReportTool.Tests/PerfReportToolTests/CommandLineTests.cs
new file mode 100644
--- /dev/null
+++ b/Engine/Source/Programs/CSVTools/PerfReportTool.Tests/PerfReportToolTests/CommandLineTests.cs
@@ -0,0 +1,281 @@
+// Copyright Epic Games, Inc. All Rights Reserved.
+
+using System;
+using System.IO;
+using System.Collections.Generic;
+using CSVStats;
+using Xunit;
+
+namespace PerfReportTool.Tests.PerfReportToolTests
+{
+	public class CommandLineParsingTests
+	{
+		#region IsParamName
+
+		[Fact]
+		public void IsParamName_DashPrefix_ReturnsTrue()
+		{
+			Assert.True(CommandLine.IsParamName("-flag"));
+		}
+
+		[Fact]
+		public void IsParamName_NoDash_ReturnsFalse()
+		{
+			Assert.False(CommandLine.IsParamName("flag"));
+		}
+
+		[Fact]
+		public void IsParamName_EmptyString_ReturnsFalse()
+		{
+			Assert.False(CommandLine.IsParamName(""));
+		}
+
+		[Fact]
+		public void IsParamName_NegativeNumber_ReturnsFalse()
+		{
+			Assert.False(CommandLine.IsParamName("-123"));
+		}
+
+		[Fact]
+		public void IsParamName_NegativeFloat_ReturnsFalse()
+		{
+			Assert.False(CommandLine.IsParamName("-3.14"));
+		}
+
+		[Fact]
+		public void IsParamName_DashOnly_ReturnsTrue()
+		{
+			// "-" is not a valid number, so it is treated as a param
+			Assert.True(CommandLine.IsParamName("-"));
+		}
+
+		#endregion
+
+		#region Basic argument retrieval
+
+		[Fact]
+		public void GetArg_CaseInsensitive_ReturnsValue()
+		{
+			var cl = new CommandLine(new[] { "-MyArg", "hello" });
+			Assert.Equal("hello", cl.GetArg("myarg"));
+			Assert.Equal("hello", cl.GetArg("MYARG"));
+			Assert.Equal("hello", cl.GetArg("MyArg"));
+		}
+
+		[Fact]
+		public void GetArg_MissingArg_ReturnsEmptyString()
+		{
+			var cl = new CommandLine(new[] { "-foo", "bar" });
+			Assert.Equal("", cl.GetArg("missing"));
+		}
+
+		[Fact]
+		public void GetArg_WithDefault_ReturnDefaultWhenMissing()
+		{
+			var cl = new CommandLine(new[] { "-foo", "bar" });
+			Assert.Equal("fallback", cl.GetArg("missing", "fallback"));
+		}
+
+		[Fact]
+		public void GetArg_FlagWithoutValue_ReturnsOne()
+		{
+			// A flag with no value after it should default to "1"
+			var cl = new CommandLine(new[] { "-verbose" });
+			Assert.Equal("1", cl.GetArg("verbose"));
+		}
+
+		[Fact]
+		public void GetArg_MultipleValues_JoinedBySemicolon()
+		{
+			// When multiple non-param values follow a param, they are joined with ";"
+			var cl = new CommandLine(new[] { "-files", "a.csv", "b.csv", "c.csv" });
+			string val = cl.GetArg("files");
+			Assert.Equal("a.csv;b.csv;c.csv", val);
+		}
+
+		#endregion
+
+		#region Typed argument retrieval
+
+		[Fact]
+		public void GetIntArg_ValidInt_ReturnsParsedValue()
+		{
+			var cl = new CommandLine(new[] { "-count", "42" });
+			Assert.Equal(42, cl.GetIntArg("count", 0));
+		}
+
+		[Fact]
+		public void GetIntArg_Missing_ReturnsDefault()
+		{
+			var cl = new CommandLine(new[] { "-other", "x" });
+			Assert.Equal(99, cl.GetIntArg("count", 99));
+		}
+
+		[Fact]
+		public void GetFloatArg_ValidFloat_ReturnsParsedValue()
+		{
+			var cl = new CommandLine(new[] { "-threshold", "3.14" });
+			Assert.Equal(3.14f, cl.GetFloatArg("threshold", 0f), 0.001f);
+		}
+
+		[Fact]
+		public void GetFloatArg_Missing_ReturnsDefault()
+		{
+			var cl = new CommandLine(new[] { "-other", "x" });
+			Assert.Equal(1.5f, cl.GetFloatArg("threshold", 1.5f), 0.001f);
+		}
+
+		#endregion
+
+		#region Boolean arguments
+
+		[Fact]
+		public void GetBoolArg_FlagPresent_ReturnsTrue()
+		{
+			var cl = new CommandLine(new[] { "-verbose" });
+			Assert.True(cl.GetBoolArg("verbose", false));
+		}
+
+		[Fact]
+		public void GetBoolArg_ExplicitZero_ReturnsFalse()
+		{
+			var cl = new CommandLine(new[] { "-verbose", "0" });
+			Assert.False(cl.GetBoolArg("verbose", true));
+		}
+
+		[Fact]
+		public void GetBoolArg_ExplicitOne_ReturnsTrue()
+		{
+			var cl = new CommandLine(new[] { "-verbose", "1" });
+			Assert.True(cl.GetBoolArg("verbose", false));
+		}
+
+		[Fact]
+		public void GetBoolArg_Missing_ReturnsDefault()
+		{
+			var cl = new CommandLine(new[] { "-other" });
+			Assert.True(cl.GetBoolArg("verbose", true));
+			Assert.False(cl.GetBoolArg("verbose", false));
+		}
+
+		[Fact]
+		public void GetOptionalBoolArg_Present_ReturnsValue()
+		{
+			var cl = new CommandLine(new[] { "-flag", "1" });
+			Assert.True(cl.GetOptionalBoolArg("flag").HasValue);
+			Assert.True(cl.GetOptionalBoolArg("flag").Value);
+		}
+
+		[Fact]
+		public void GetOptionalBoolArg_PresentZero_ReturnsFalse()
+		{
+			var cl = new CommandLine(new[] { "-flag", "0" });
+			Assert.True(cl.GetOptionalBoolArg("flag").HasValue);
+			Assert.False(cl.GetOptionalBoolArg("flag").Value);
+		}
+
+		[Fact]
+		public void GetOptionalBoolArg_Missing_ReturnsNull()
+		{
+			var cl = new CommandLine(new[] { "-other" });
+			Assert.Null(cl.GetOptionalBoolArg("flag"));
+		}
+
+		#endregion
+
+		#region Negative value encoding
+
+		[Fact]
+		public void NegativeValue_MinusEncoded_DecodesCorrectly()
+		{
+			// The tool encodes negative values as "&minus;" to avoid them being parsed as flags
+			var cl = new CommandLine(new[] { "-offset", "&minus;100" });
+			Assert.Equal("-100", cl.GetArg("offset"));
+			Assert.Equal(-100, cl.GetIntArg("offset", 0));
+		}
+
+		#endregion
+
+		#region Duplicate argument handling
+
+		[Fact]
+		public void DuplicateArg_LastValueWins()
+		{
+			var cl = new CommandLine(new[] { "-name", "first", "-name", "second" });
+			Assert.Equal("second", cl.GetArg("name"));
+		}
+
+		#endregion
+
+		#region Response file
+
+		[Fact]
+		[Trait("Category", "FileIO")]
+		public void ResponseFile_ReadsArgsFromFile()
+		{
+			string tempFile = Path.GetTempFileName();
+			try
+			{
+				File.WriteAllLines(tempFile, new[]
+				{
+					"-csvdir /some/path",
+					"-nodetailedreports",
+					"-o /output"
+				});
+
+				var cl = new CommandLine(new[] { "-response", tempFile });
+				Assert.Equal("/some/path", cl.GetArg("csvdir"));
+				Assert.Equal("1", cl.GetArg("nodetailedreports"));
+				Assert.Equal("/output", cl.GetArg("o"));
+			}
+			finally
+			{
+				File.Delete(tempFile);
+			}
+		}
+
+		#endregion
+
+		#region GetCommandLine
+
+		[Fact]
+		public void GetCommandLine_ReconstructsString()
+		{
+			var cl = new CommandLine(new[] { "-csv", "test.csv", "-o", "output" });
+			string cmdLine = cl.GetCommandLine();
+			Assert.Contains("-csv", cmdLine);
+			Assert.Contains("test.csv", cmdLine);
+			Assert.Contains("-o", cmdLine);
+			Assert.Contains("output", cmdLine);
+		}
+
+		[Fact]
+		public void GetCommandLine_QuotesArgsWithSpaces()
+		{
+			var cl = new CommandLine(new[] { "-csv", "path with spaces/test.csv" });
+			string cmdLine = cl.GetCommandLine();
+			Assert.Contains("\"path with spaces/test.csv\"", cmdLine);
+		}
+
+		#endregion
+	}
+
+	public class PerfLogTests
+	{
+		[Fact]
+		public void PerfLog_LogTiming_ReturnsElapsedTime()
+		{
+			var log = new PerfLog(false);
+			double elapsed = log.LogTiming("test");
+			Assert.True(elapsed >= 0.0);
+		}
+
+		[Fact]
+		public void PerfLog_LogTotalTiming_ReturnsPositive()
+		{
+			var log = new PerfLog(false);
+			double total = log.LogTotalTiming();
+			Assert.True(total >= 0.0);
+		}
+	}
+}
diff --git a/Engine/Source/Programs/CSVTools/PerfReportTool.Tests/PerfReportToolTests/CsvFileCacheTests.cs b/Engine/Source/Programs/CSVTools/PerfReportTool.Tests/PerfReportToolTests/CsvFileCacheTests.cs
new file mode 100644
--- /dev/null
+++ b/Engine/Source/Programs/CSVTools/PerfReportTool.Tests/PerfReportToolTests/CsvFileCacheTests.cs
@@ -0,0 +1,74 @@
+// Copyright Epic Games, Inc. All Rights Reserved.
+
+using System;
+using System.Collections.Generic;
+using System.IO;
+using CSVStats;
+using PerfReportTool;
+using PerfSummaries;
+using Xunit;
+
+namespace PerfReportTool.Tests.PerfReportToolTests
+{
+	public class ExternalMetadataCacheTests
+	{
+		[Fact]
+		[Trait("Category", "FileIO")]
+		public void InjectExternalMetadata_AddsMetadata()
+		{
+			string tempDir = Path.Combine(Path.GetTempPath(), "ExtMeta_" + Path.GetRandomFileName());
+			Directory.CreateDirectory(tempDir);
+			try
+			{
+				string metaFile = Path.Combine(tempDir, "test.txt");
+				File.WriteAllLines(metaFile, new[]
+				{
+					"custom_key=custom_value",
+					"another_key=another_value"
+				});
+
+				var cache = new ExternalMetadataCache(new List<string> { metaFile });
+				var rowData = new SummaryTableRowData();
+				rowData.Add(SummaryTableElement.Type.CsvMetadata, "csvid", "test123");
+
+				cache.InjectExternalMetadata(rowData);
+
+				Assert.True(rowData.Contains("custom_key"), "External metadata key should be injected");
+			}
+			finally
+			{
+				Directory.Delete(tempDir, true);
+			}
+		}
+
+		[Fact]
+		[Trait("Category", "FileIO")]
+		public void InjectExternalMetadata_DoesNotOverwriteExisting()
+		{
+			string tempDir = Path.Combine(Path.GetTempPath(), "ExtMeta_" + Path.GetRandomFileName());
+			Directory.CreateDirectory(tempDir);
+			try
+			{
+				string metaFile = Path.Combine(tempDir, "test.txt");
+				File.WriteAllLines(metaFile, new[]
+				{
+					"platform=external_value"
+				});
+
+				var cache = new ExternalMetadataCache(new List<string> { metaFile });
+				var rowData = new SummaryTableRowData();
+				rowData.Add(SummaryTableElement.Type.CsvMetadata, "platform", "original_value");
+				rowData.Add(SummaryTableElement.Type.CsvMetadata, "csvid", "test123");
+
+				cache.InjectExternalMetadata(rowData);
+
+				// The existing value should not be overwritten
+				Assert.Equal("original_value", rowData.Get("platform").value);
+			}
+			finally
+			{
+				Directory.Delete(tempDir, true);
+			}
+		}
+	}
+}
diff --git a/Engine/Source/Programs/CSVTools/PerfReportTool.Tests/PerfReportToolTests/HashHelperTests.cs b/Engine/Source/Programs/CSVTools/PerfReportTool.Tests/PerfReportToolTests/HashHelperTests.cs
new file mode 100644
--- /dev/null
+++ b/Engine/Source/Programs/CSVTools/PerfReportTool.Tests/PerfReportToolTests/HashHelperTests.cs
@@ -0,0 +1,51 @@
+using PerfReportTool;
+using Xunit;
+
+namespace PerfReportTool.Tests.PerfReportToolTests
+{
+	public class HashHelperTests
+	{
+		[Fact]
+		public void StringToHashStr_SameInput_ProducesSameOutput()
+		{
+			string input = "test string for hashing";
+			string hash1 = HashHelper.StringToHashStr(input);
+			string hash2 = HashHelper.StringToHashStr(input);
+			Assert.Equal(hash1, hash2);
+		}
+
+		[Fact]
+		public void StringToHashStr_DifferentInput_ProducesDifferentOutput()
+		{
+			string hash1 = HashHelper.StringToHashStr("input A");
+			string hash2 = HashHelper.StringToHashStr("input B");
+			Assert.NotEqual(hash1, hash2);
+		}
+
+		[Fact]
+		public void StringToHashStr_ReturnsHexString()
+		{
+			string hash = HashHelper.StringToHashStr("test");
+			// SHA256 produces 32 bytes = 64 hex characters
+			Assert.Equal(64, hash.Length);
+			Assert.Matches("^[0-9A-Fa-f]+$", hash);
+		}
+
+		[Fact]
+		public void StringToHashStr_MaxCharsOut_TruncatesCorrectly()
+		{
+			string hashFull = HashHelper.StringToHashStr("test");
+			string hashTruncated = HashHelper.StringToHashStr("test", 16);
+			Assert.Equal(16, hashTruncated.Length);
+			Assert.Equal(hashFull.Substring(0, 16), hashTruncated);
+		}
+
+		[Fact]
+		public void StringToHashStr_EmptyString_DoesNotThrow()
+		{
+			string hash = HashHelper.StringToHashStr("");
+			Assert.NotNull(hash);
+			Assert.NotEmpty(hash);
+		}
+	}
+}
diff --git a/Engine/Source/Programs/CSVTools/PerfReportTool.Tests/PerfReportToolTests/ReportXmlLoadTests.cs b/Engine/Source/Programs/CSVTools/PerfReportTool.Tests/PerfReportToolTests/ReportXmlLoadTests.cs
new file mode 100644
--- /dev/null
+++ b/Engine/Source/Programs/CSVTools/PerfReportTool.Tests/PerfReportToolTests/ReportXmlLoadTests.cs
@@ -0,0 +1,80 @@
+using System.IO;
+using System.Linq;
+using System.Xml.Linq;
+using PerfReportTool.Tests.TestHelpers;
+using Xunit;
+
+namespace PerfReportTool.Tests.PerfReportToolTests
+{
+	[Trait("Category", "FileIO")]
+	public class ReportXmlLoadTests
+	{
+		[Fact]
+		public void ReportTypesXml_Exists_InOutputDir()
+		{
+			string path = ReportXmlTestHelper.GetReportTypesXmlPath();
+			Assert.True(File.Exists(path), $"ReportTypes.xml not found at {path}");
+		}
+
+		[Fact]
+		public void ReportGraphsXml_Exists_InOutputDir()
+		{
+			string path = ReportXmlTestHelper.GetReportGraphsXmlPath();
+			Assert.True(File.Exists(path), $"ReportGraphs.xml not found at {path}");
+		}
+
+		[Fact]
+		public void LlmReportTypesXml_Exists_InOutputDir()
+		{
+			string path = ReportXmlTestHelper.GetLlmReportTypesXmlPath();
+			Assert.True(File.Exists(path), $"LLMReportTypes.xml not found at {path}");
+		}
+
+		[Fact]
+		public void LlmReportGraphsXml_Exists_InOutputDir()
+		{
+			string path = ReportXmlTestHelper.GetLlmReportGraphsXmlPath();
+			Assert.True(File.Exists(path), $"LLMReportGraphs.xml not found at {path}");
+		}
+
+		[Fact]
+		public void ReportTypesXml_ContainsReportTypeDefinitions()
+		{
+			string path = ReportXmlTestHelper.GetReportTypesXmlPath();
+			Assert.True(File.Exists(path), $"ReportTypes.xml not found at {path}");
+
+			var doc = XDocument.Load(path);
+			var root = doc.Root;
+			Assert.NotNull(root);
+
+			// Should contain reporttype elements defining how CSVs are processed
+			var ns = root.GetDefaultNamespace();
+			var reportTypes = root.Descendants().Where(e => e.Name.LocalName == "reporttype").ToList();
+			Assert.True(reportTypes.Count > 0, "ReportTypes.xml should contain at least one <reporttype> element");
+
+			// Each reporttype should have a name
+			foreach (var rt in reportTypes)
+			{
+				var nameAttr = rt.Attribute("name");
+				Assert.NotNull(nameAttr);
+				Assert.False(string.IsNullOrEmpty(nameAttr.Value), "reporttype should have a non-empty name attribute");
+			}
+		}
+
+		[Fact]
+		public void ReportGraphsXml_ContainsGraphDefinitions()
+		{
+			string path = ReportXmlTestHelper.GetReportGraphsXmlPath();
+			Assert.True(File.Exists(path), $"ReportGraphs.xml not found at {path}");
+
+			var doc = XDocument.Load(path);
+			var root = doc.Root;
+			Assert.NotNull(root);
+
+			// Should contain graph group elements
+			var graphGroups = root.Descendants().Where(e =>
+				e.Name.LocalName == "graphGroup" || e.Name.LocalName == "graph").ToList();
+			Assert.True(graphGroups.Count > 0, "ReportGraphs.xml should contain graph definitions");
+		}
+	}
+}
diff --git a/Engine/Source/Programs/CSVTools/PerfReportTool.Tests/PerfReportToolTests/ReportXmlTests.cs b/Engine/Source/Programs/CSVTools/PerfReportTool.Tests/PerfReportToolTests/ReportXmlTests.cs
new file mode 100644
--- /dev/null
+++ b/Engine/Source/Programs/CSVTools/PerfReportTool.Tests/PerfReportToolTests/ReportXmlTests.cs
@@ -0,0 +1,235 @@
+// Copyright Epic Games, Inc. All Rights Reserved.
+
+using System;
+using System.Collections.Generic;
+using System.IO;
+using System.Xml.Linq;
+using PerfReportTool;
+using PerfSummaries;
+using Xunit;
+
+namespace PerfReportTool.Tests.PerfReportToolTests
+{
+	public class XmlVariableMappingsExtendedTests
+	{
+		[Fact]
+		public void SetVariable_StoresAndRetrieves()
+		{
+			var vars = new XmlVariableMappings();
+			vars.SetVariable("testVar", "hello");
+			string result = vars.ResolveVariables("prefix_${testVar}_suffix");
+			Assert.Equal("prefix_hello_suffix", result);
+		}
+
+		[Fact]
+		public void ResolveVariables_MultipleVariables_AllResolved()
+		{
+			var vars = new XmlVariableMappings();
+			vars.SetVariable("a", "1");
+			vars.SetVariable("b", "2");
+			string result = vars.ResolveVariables("${a}+${b}");
+			Assert.Equal("1+2", result);
+		}
+
+		[Fact]
+		public void ResolveVariables_MissingVariable_ReplacedWithEmpty()
+		{
+			var vars = new XmlVariableMappings();
+			string result = vars.ResolveVariables("prefix_${missing}_suffix");
+			Assert.Equal("prefix__suffix", result);
+		}
+
+		[Fact]
+		public void ResolveVariables_NoVariables_ReturnsOriginal()
+		{
+			var vars = new XmlVariableMappings();
+			string result = vars.ResolveVariables("no variables here");
+			Assert.Equal("no variables here", result);
+		}
+
+		[Fact]
+		public void ResolveVariables_ArrayIndexing_ReturnsElement()
+		{
+			var vars = new XmlVariableMappings();
+			vars.SetVariable("list", "alpha,beta,gamma");
+			string result = vars.ResolveVariables("${list[1]}");
+			Assert.Equal("beta", result);
+		}
+
+		[Fact]
+		public void ResolveVariables_ArrayIndexZero_ReturnsFirst()
+		{
+			var vars = new XmlVariableMappings();
+			vars.SetVariable("items", "first,second,third");
+			string result = vars.ResolveVariables("${items[0]}");
+			Assert.Equal("first", result);
+		}
+
+		[Fact]
+		public void SetMetadataVariables_SetsMetaPrefix()
+		{
+			var vars = new XmlVariableMappings();
+			var metadata = new CSVStats.CsvMetadata();
+			metadata.Values["platform"] = "windows";
+			metadata.Values["buildversion"] = "5.4";
+
+			vars.SetMetadataVariables(metadata);
+
+			string result = vars.ResolveVariables("${meta.platform}");
+			Assert.Equal("windows", result);
+		}
+
+		[Fact]
+		public void ApplyVariableSet_ParsesVarElements()
+		{
+			// New format: <var name="name">value</var> (value as element text)
+			string xml = @"<variableSet>
+				<var name='targetFps'>30</var>
+				<var name='hitchMs'>150</var>
+			</variableSet>";
+			var element = XElement.Parse(xml);
+			var vars = new XmlVariableMappings();
+
+			vars.ApplyVariableSet(element, null);
+
+			Assert.Equal("30", vars.ResolveVariables("${targetFps}"));
+			Assert.Equal("150", vars.ResolveVariables("${hitchMs}"));
+		}
+
+		[Fact]
+		public void ApplyVariableSet_WithMultiplier_ScalesNumericValues()
+		{
+			string xml = @"<variableSet multiplier='2.0'>
+				<var name='budget'>33.33</var>
+			</variableSet>";
+			var element = XElement.Parse(xml);
+			var vars = new XmlVariableMappings();
+
+			vars.ApplyVariableSet(element, null, 1.0);
+
+			string result = vars.ResolveVariables("${budget}");
+			double parsed = double.Parse(result, System.Globalization.CultureInfo.InvariantCulture);
+			Assert.Equal(66.66, parsed, 0.01);
+		}
+	}
+
+	[Collection("SummaryFactory")]
+	public class ReportTypeInfoTests
+	{
+		// FPSChartSummary has the (XElement, XmlVariableMappings, string) constructor that SummaryFactory.Create requires
+		private const string FpsChartSummaryXml = "<summary type='fpschart' fps='30' hitchThreshold='150'><stats/></summary>";
+
+		[Fact]
+		public void GetSummaryTableCacheID_ReturnsNonEmptyString()
+		{
+			string xml = $"<reporttype name='test' title='Test Report'>{FpsChartSummaryXml}</reporttype>";
+			var element = XElement.Parse(xml);
+			var vars = new XmlVariableMappings();
+
+			SummaryFactory.Init();
+			var reportType = new ReportTypeInfo(element, new Dictionary<string, XElement>(), "", vars, null);
+
+			string cacheId = reportType.GetSummaryTableCacheID();
+			Assert.False(string.IsNullOrEmpty(cacheId));
+		}
+
+		[Fact]
+		public void Constructor_ParsesNameAndTitle()
+		{
+			string xml = $"<reporttype name='flythrough' title='Flythrough Report'>{FpsChartSummaryXml}</reporttype>";
+			var element = XElement.Parse(xml);
+			var vars = new XmlVariableMappings();
+
+			SummaryFactory.Init();
+			var reportType = new ReportTypeInfo(element, new Dictionary<string, XElement>(), "", vars, null);
+
+			Assert.Equal("flythrough", reportType.name);
+			Assert.Equal("Flythrough Report", reportType.title);
+		}
+
+		[Fact]
+		public void Constructor_ParsesStripEvents()
+		{
+			string xml = $"<reporttype name='test' title='Test' stripEvents='true'>{FpsChartSummaryXml}</reporttype>";
+			var element = XElement.Parse(xml);
+			var vars = new XmlVariableMappings();
+
+			SummaryFactory.Init();
+			var reportType = new ReportTypeInfo(element, new Dictionary<string, XElement>(), "", vars, null);
+
+			Assert.True(reportType.bStripEvents);
+		}
+
+		[Fact]
+		public void Constructor_ParsesMetadataToShow()
+		{
+			string xml = $@"<reporttype name='test' title='Test'>
+				<metadataToShow>platform,buildversion,deviceprofile</metadataToShow>
+				{FpsChartSummaryXml}
+			</reporttype>";
+			var element = XElement.Parse(xml);
+			var vars = new XmlVariableMappings();
+
+			SummaryFactory.Init();
+			var reportType = new ReportTypeInfo(element, new Dictionary<string, XElement>(), "", vars, null);
+
+			Assert.NotNull(reportType.metadataToShowList);
+			Assert.Equal(3, reportType.metadataToShowList.Length);
+			Assert.Contains("platform", reportType.metadataToShowList);
+		}
+
+		[Fact]
+		public void Constructor_CreatesSummaries()
+		{
+			string xml = $"<reporttype name='test' title='Test'>{FpsChartSummaryXml}</reporttype>";
+			var element = XElement.Parse(xml);
+			var vars = new XmlVariableMappings();
+
+			SummaryFactory.Init();
+			var reportType = new ReportTypeInfo(element, new Dictionary<string, XElement>(), "", vars, null);
+
+			Assert.True(reportType.summaries.Count > 0, "Should have at least one summary");
+		}
+
+		[Fact]
+		public void GetSummaryTableCacheID_DifferentReportTypes_DifferentIDs()
+		{
+			SummaryFactory.Init();
+			var vars = new XmlVariableMappings();
+
+			string xml1 = $"<reporttype name='type1' title='Type 1'>{FpsChartSummaryXml}</reporttype>";
+			string xml2 = $"<reporttype name='type2' title='Type 2'>{FpsChartSummaryXml}</reporttype>";
+
+			var rt1 = new ReportTypeInfo(XElement.Parse(xml1), new Dictionary<string, XElement>(), "", vars, null);
+			var rt2 = new ReportTypeInfo(XElement.Parse(xml2), new Dictionary<string, XElement>(), "", vars, null);
+
+			Assert.NotEqual(rt1.GetSummaryTableCacheID(), rt2.GetSummaryTableCacheID());
+		}
+	}
+
+	public class SummaryTableInfoTests
+	{
+		[Fact]
+		public void Constructor_ParsesFilterAndRowSort()
+		{
+			var info = new SummaryTableInfo("frametime,drawcalls,platform", "buildversion,platform");
+			Assert.Contains("frametime", info.columnFilterList);
+			Assert.Contains("drawcalls", info.columnFilterList);
+			Assert.Contains("platform", info.columnFilterList);
+			Assert.Contains("buildversion", info.rowSortList);
+			Assert.Contains("platform", info.rowSortList);
+		}
+
+		[Fact]
+		public void DefaultValues_AreCorrect()
+		{
+			var info = new SummaryTableInfo();
+			Assert.Equal(TableColorizeMode.Budget, info.tableColorizeMode);
+			Assert.Equal(DateCollationVisibility.Newest, info.dateCollationVisibility);
+			Assert.Equal(StringCollationVisibility.Auto, info.stringCollationVisibility);
+			Assert.Equal(TableColumnSortMode.Default, info.columnSortMode);
+			Assert.False(info.bReverseSortRows);
+			Assert.False(info.bScrollableFormatting);
+		}
+	}
+}
diff --git a/Engine/Source/Programs/CSVTools/PerfReportTool.Tests/PerfReportToolTests/SummaryFactoryTests.cs b/Engine/Source/Programs/CSVTools/PerfReportTool.Tests/PerfReportToolTests/SummaryFactoryTests.cs
new file mode 100644
--- /dev/null
+++ b/Engine/Source/Programs/CSVTools/PerfReportTool.Tests/PerfReportToolTests/SummaryFactoryTests.cs
@@ -0,0 +1,61 @@
+using System.Collections.Generic;
+using System.Reflection;
+using PerfSummaries;
+using Xunit;
+
+namespace PerfReportTool.Tests.PerfReportToolTests
+{
+	[Collection("SummaryFactory")]
+	public class SummaryFactoryTests
+	{
+		public SummaryFactoryTests()
+		{
+			SummaryFactory.Init();
+		}
+
+		private static Dictionary<string, System.Type> GetRegisteredTypes()
+		{
+			// Access the private static summaryNameLookup field via reflection
+			var field = typeof(SummaryFactory).GetField("summaryNameLookup", BindingFlags.NonPublic | BindingFlags.Static);
+			Assert.NotNull(field);
+			return (Dictionary<string, System.Type>)field.GetValue(null);
+		}
+
+		[Theory]
+		[InlineData("fpschart")]
+		[InlineData("hitches")]
+		[InlineData("histogram")]
+		[InlineData("event")]
+		[InlineData("peak")]
+		[InlineData("bucketsummary")]
+		[InlineData("boundedstatvalues")]
+		[InlineData("checkpoint")]
+		[InlineData("mapoverlay")]
+		[InlineData("statBudgets")]
+		[InlineData("timingregion")]
+		[InlineData("standard")]
+		[InlineData("extralinks")]
+		public void Init_RegistersKnownSummaryType(string summaryName)
+		{
+			var registry = GetRegisteredTypes();
+			Assert.True(registry.ContainsKey(summaryName), $"Summary type '{summaryName}' was not registered in SummaryFactory");
+			Assert.True(registry[summaryName].IsSubclassOf(typeof(Summary)),
+				$"Registered type for '{summaryName}' should be a Summary subclass");
+		}
+
+		[Fact]
+		public void Init_RegistersAtLeast13Types()
+		{
+			var registry = GetRegisteredTypes();
+			Assert.True(registry.Count >= 13, $"Expected at least 13 summary types but found {registry.Count}");
+		}
+
+		[Fact]
+		public void Create_UnknownType_ThrowsException()
+		{
+			var ex = Assert.Throws<System.Exception>(() =>
+				SummaryFactory.Create("NonExistentSummaryType", null, null, ""));
+			Assert.Contains("not found", ex.Message);
+		}
+	}
+}
diff --git a/Engine/Source/Programs/CSVTools/PerfReportTool.Tests/PerfReportToolTests/SummaryTableElementTests.cs b/Engine/Source/Programs/CSVTools/PerfReportTool.Tests/PerfReportToolTests/SummaryTableElementTests.cs
new file mode 100644
--- /dev/null
+++ b/Engine/Source/Programs/CSVTools/PerfReportTool.Tests/PerfReportToolTests/SummaryTableElementTests.cs
@@ -0,0 +1,220 @@
+using System.IO;
+using PerfSummaries;
+using Xunit;
+
+namespace PerfReportTool.Tests.PerfReportToolTests
+{
+	public class SummaryTableElementTests
+	{
+		[Fact]
+		public void NumericConstructor_SetsProperties()
+		{
+			var element = new SummaryTableElement(
+				SummaryTableElement.Type.CsvStatAverage,
+				"TestStat",
+				42.5,
+				null,
+				"tooltip text");
+
+			Assert.Equal(SummaryTableElement.Type.CsvStatAverage, element.type);
+			Assert.Equal("TestStat", element.name);
+			Assert.True(element.isNumeric);
+			Assert.Equal(42.5, element.numericValue);
+			Assert.Equal("42.5", element.value);
+			Assert.Equal("tooltip text", element.tooltip);
+		}
+
+		[Fact]
+		public void StringConstructor_SetsProperties()
+		{
+			var element = new SummaryTableElement(
+				SummaryTableElement.Type.CsvMetadata,
+				"platform",
+				"windows",
+				null,
+				"");
+
+			Assert.Equal(SummaryTableElement.Type.CsvMetadata, element.type);
+			Assert.Equal("platform", element.name);
+			Assert.False(element.isNumeric);
+			Assert.Equal("windows", element.value);
+		}
+
+		[Fact]
+		public void Constructor_EmptyName_ThrowsException()
+		{
+			Assert.Throws<System.Exception>(() =>
+				new SummaryTableElement(SummaryTableElement.Type.CsvStatAverage, "", 0.0, null, ""));
+		}
+
+		[Fact]
+		public void CacheRoundTrip_NumericElement_PreservesData()
+		{
+			var thresholds = new ColourThresholdList(40, 30, 20, 10);
+			var original = new SummaryTableElement(
+				SummaryTableElement.Type.CsvStatAverage,
+				"frametime",
+				16.67,
+				thresholds,
+				"Frame time in ms",
+				(uint)SummaryTableElement.Flags.Hidden);
+
+			using var stream = new MemoryStream();
+			using (var writer = new BinaryWriter(stream, System.Text.Encoding.UTF8, leaveOpen: true))
+			{
+				original.WriteToCache(writer);
+			}
+
+			stream.Position = 0;
+			using var reader = new BinaryReader(stream);
+			var restored = SummaryTableElement.ReadFromCache(reader);
+
+			Assert.Equal(original.type, restored.type);
+			Assert.Equal(original.name, restored.name);
+			Assert.Equal(original.value, restored.value);
+			Assert.Equal(original.tooltip, restored.tooltip);
+			Assert.Equal(original.numericValue, restored.numericValue);
+			Assert.Equal(original.isNumeric, restored.isNumeric);
+			Assert.Equal(original.flags, restored.flags);
+			Assert.NotNull(restored.colorThresholdList);
+			Assert.Equal(original.colorThresholdList.Count, restored.colorThresholdList.Count);
+			for (int i = 0; i < original.colorThresholdList.Count; i++)
+			{
+				Assert.Equal(original.colorThresholdList.Thresholds[i].value, restored.colorThresholdList.Thresholds[i].value);
+			}
+		}
+
+		[Fact]
+		public void CacheRoundTrip_StringElement_PreservesData()
+		{
+			var original = new SummaryTableElement(
+				SummaryTableElement.Type.CsvMetadata,
+				"buildversion",
+				"++UE5+Main-CL-12345",
+				null,
+				"");
+
+			using var stream = new MemoryStream();
+			using (var writer = new BinaryWriter(stream, System.Text.Encoding.UTF8, leaveOpen: true))
+			{
+				original.WriteToCache(writer);
+			}
+
+			stream.Position = 0;
+			using var reader = new BinaryReader(stream);
+			var restored = SummaryTableElement.ReadFromCache(reader);
+
+			Assert.Equal(original.type, restored.type);
+			Assert.Equal(original.name, restored.name);
+			Assert.Equal(original.value, restored.value);
+			Assert.False(restored.isNumeric);
+			Assert.Null(restored.colorThresholdList);
+		}
+
+		[Fact]
+		public void SetFlag_GetFlag_WorkCorrectly()
+		{
+			var element = new SummaryTableElement(
+				SummaryTableElement.Type.CsvStatAverage, "test", 1.0, null, "");
+
+			Assert.False(element.GetFlag(SummaryTableElement.Flags.Hidden));
+
+			element.SetFlag(SummaryTableElement.Flags.Hidden, true);
+			Assert.True(element.GetFlag(SummaryTableElement.Flags.Hidden));
+
+			element.SetFlag(SummaryTableElement.Flags.Hidden, false);
+			Assert.False(element.GetFlag(SummaryTableElement.Flags.Hidden));
+		}
+
+		[Fact]
+		public void ToJsonDict_NumericValue_SerializesCorrectly()
+		{
+			var element = new SummaryTableElement(
+				SummaryTableElement.Type.CsvStatAverage, "frametime", 16.67, null, "");
+
+			var dict = element.ToJsonDict(true);
+
+			Assert.Equal("CsvStatAverage", dict["type"]);
+			Assert.IsType<decimal>(dict["value"]);
+			Assert.Equal(16.67m, (decimal)dict["value"], 2);
+		}
+
+		[Fact]
+		public void ToJsonDict_StringValue_SerializesCorrectly()
+		{
+			var element = new SummaryTableElement(
+				SummaryTableElement.Type.CsvMetadata, "platform", "windows", null, "");
+
+			var dict = element.ToJsonDict(true);
+
+			Assert.Equal("CsvMetadata", dict["type"]);
+			Assert.Equal("windows", dict["value"]);
+		}
+
+		[Fact]
+		public void ToJsonDict_WithTooltip_IncludesIt()
+		{
+			var element = new SummaryTableElement(
+				SummaryTableElement.Type.CsvStatAverage, "stat", 1.0, null, "my tooltip");
+
+			var dict = element.ToJsonDict(false);
+			Assert.True(dict.ContainsKey("tooltip"));
+			Assert.Equal("my tooltip", dict["tooltip"]);
+		}
+
+		[Fact]
+		public void ToJsonDict_EmptyTooltip_ExcludesIt()
+		{
+			var element = new SummaryTableElement(
+				SummaryTableElement.Type.CsvStatAverage, "stat", 1.0, null, "");
+
+			var dict = element.ToJsonDict(false);
+			Assert.False(dict.ContainsKey("tooltip"));
+		}
+
+		[Fact]
+		public void ToJsonDict_WithFlags_IncludesFlagStrings()
+		{
+			var element = new SummaryTableElement(
+				SummaryTableElement.Type.CsvStatAverage, "stat", 1.0, null, "",
+				(uint)SummaryTableElement.Flags.Hidden);
+
+			var dict = element.ToJsonDict(false);
+			Assert.True(dict.ContainsKey("flags"));
+		}
+
+		[Fact]
+		public void Clone_ProducesIndependentCopy()
+		{
+			var original = new SummaryTableElement(
+				SummaryTableElement.Type.CsvStatAverage, "stat", 42.0, null, "tip");
+			var clone = original.Clone();
+
+			Assert.Equal(original.name, clone.name);
+			Assert.Equal(original.numericValue, clone.numericValue);
+			Assert.NotSame(original, clone);
+		}
+
+		[Fact]
+		public void DynamicValue_Numeric_ReturnsDouble()
+		{
+			var element = new SummaryTableElement(
+				SummaryTableElement.Type.CsvStatAverage, "stat", 42.5, null, "");
+
+			dynamic val = element.DynamicValue;
+			Assert.IsType<double>(val);
+			Assert.Equal(42.5, (double)val);
+		}
+
+		[Fact]
+		public void DynamicValue_String_ReturnsString()
+		{
+			var element = new SummaryTableElement(
+				SummaryTableElement.Type.CsvMetadata, "platform", "windows", null, "");
+
+			dynamic val = element.DynamicValue;
+			Assert.IsType<string>(val);
+			Assert.Equal("windows", (string)val);
+		}
+	}
+}
diff --git a/Engine/Source/Programs/CSVTools/PerfReportTool.Tests/PerfReportToolTests/SummaryTableRowDataTests.cs b/Engine/Source/Programs/CSVTools/PerfReportTool.Tests/PerfReportToolTests/SummaryTableRowDataTests.cs
new file mode 100644
--- /dev/null
+++ b/Engine/Source/Programs/CSVTools/PerfReportTool.Tests/PerfReportToolTests/SummaryTableRowDataTests.cs
@@ -0,0 +1,272 @@
+using System.Collections.Generic;
+using System.IO;
+using PerfSummaries;
+using Xunit;
+
+namespace PerfReportTool.Tests.PerfReportToolTests
+{
+	public class SummaryTableRowDataTests
+	{
+		[Fact]
+		public void Add_NumericEntry_CanBeRetrieved()
+		{
+			var rowData = new SummaryTableRowData();
+			rowData.Add(SummaryTableElement.Type.CsvStatAverage, "frametime", 16.67);
+
+			Assert.True(rowData.Contains("frametime"));
+			var element = rowData.Get("frametime");
+			Assert.NotNull(element);
+			Assert.Equal(16.67, element.numericValue);
+			Assert.True(element.isNumeric);
+		}
+
+		[Fact]
+		public void Add_StringEntry_CanBeRetrieved()
+		{
+			var rowData = new SummaryTableRowData();
+			rowData.Add(SummaryTableElement.Type.CsvMetadata, "platform", "windows");
+
+			Assert.True(rowData.Contains("platform"));
+			var element = rowData.Get("platform");
+			Assert.NotNull(element);
+			Assert.Equal("windows", element.value);
+		}
+
+		[Fact]
+		public void Add_KeyIsCaseInsensitive()
+		{
+			var rowData = new SummaryTableRowData();
+			rowData.Add(SummaryTableElement.Type.CsvStatAverage, "FrameTime", 16.67);
+
+			// Keys are lowercased on add
+			Assert.True(rowData.Contains("frametime"));
+			Assert.NotNull(rowData.Get("frametime"));
+		}
+
+		[Fact]
+		public void Add_DuplicateKey_DoesNotThrow()
+		{
+			var rowData = new SummaryTableRowData();
+			rowData.Add(SummaryTableElement.Type.CsvStatAverage, "stat", 1.0);
+			// Should not throw, just warn
+			rowData.Add(SummaryTableElement.Type.CsvStatAverage, "stat", 2.0);
+		}
+
+		[Fact]
+		public void Add_StringNumericValue_ParsedAsNumeric()
+		{
+			var rowData = new SummaryTableRowData();
+			rowData.Add(SummaryTableElement.Type.CsvStatAverage, "stat", "42.5");
+
+			var element = rowData.Get("stat");
+			Assert.True(element.isNumeric);
+			Assert.Equal(42.5, element.numericValue);
+		}
+
+		[Fact]
+		public void Add_CsvId_TreatedAsString()
+		{
+			var rowData = new SummaryTableRowData();
+			rowData.Add(SummaryTableElement.Type.CsvMetadata, "csvid", "12345");
+
+			var element = rowData.Get("csvid");
+			// csvid should always be string even if it looks numeric
+			Assert.False(element.isNumeric);
+			Assert.Equal("12345", element.value);
+		}
+
+		[Fact]
+		public void GetFrameCount_ReturnsCorrectCount()
+		{
+			var rowData = new SummaryTableRowData();
+			rowData.Add(SummaryTableElement.Type.ToolMetadata, "framecount", 1000.0);
+
+			Assert.Equal(1000, rowData.GetFrameCount());
+		}
+
+		[Fact]
+		public void GetFrameCount_NoEntry_ReturnsZero()
+		{
+			var rowData = new SummaryTableRowData();
+			Assert.Equal(0, rowData.GetFrameCount());
+		}
+
+		[Fact]
+		public void Get_MissingKey_ReturnsNull()
+		{
+			var rowData = new SummaryTableRowData();
+			Assert.Null(rowData.Get("nonexistent"));
+		}
+
+		[Fact]
+		public void RemoveSafe_ExistingKey_Removes()
+		{
+			var rowData = new SummaryTableRowData();
+			rowData.Add(SummaryTableElement.Type.CsvStatAverage, "stat", 1.0);
+			Assert.True(rowData.Contains("stat"));
+
+			rowData.RemoveSafe("stat");
+			Assert.False(rowData.Contains("stat"));
+		}
+
+		[Fact]
+		public void RemoveSafe_MissingKey_DoesNotThrow()
+		{
+			var rowData = new SummaryTableRowData();
+			rowData.RemoveSafe("nonexistent");
+		}
+
+		[Fact]
+		public void ReadCsvMetadata_ExtractsOnlyMetadata()
+		{
+			var rowData = new SummaryTableRowData();
+			rowData.Add(SummaryTableElement.Type.CsvMetadata, "platform", "windows");
+			rowData.Add(SummaryTableElement.Type.CsvMetadata, "buildversion", "++UE5+Main-CL-12345");
+			rowData.Add(SummaryTableElement.Type.CsvStatAverage, "frametime", 16.67);
+			rowData.Add(SummaryTableElement.Type.ToolMetadata, "framecount", 1000.0);
+
+			var metadata = rowData.ReadCsvMetadata();
+			Assert.NotNull(metadata);
+			Assert.Equal(2, metadata.Values.Count);
+			Assert.Equal("windows", metadata.Values["platform"]);
+			Assert.Equal("++UE5+Main-CL-12345", metadata.Values["buildversion"]);
+		}
+
+		[Fact]
+		public void ReadCsvMetadata_NoMetadata_ReturnsNull()
+		{
+			var rowData = new SummaryTableRowData();
+			rowData.Add(SummaryTableElement.Type.CsvStatAverage, "frametime", 16.67);
+
+			var metadata = rowData.ReadCsvMetadata();
+			Assert.Null(metadata);
+		}
+
+		[Fact]
+		[Trait("Category", "FileIO")]
+		public void CacheRoundTrip_PreservesAllData()
+		{
+			var rowData = new SummaryTableRowData();
+			rowData.Add(SummaryTableElement.Type.CsvMetadata, "csvid", "test-csv-id");
+			rowData.Add(SummaryTableElement.Type.CsvMetadata, "platform", "windows");
+			rowData.Add(SummaryTableElement.Type.CsvStatAverage, "frametime", 16.67);
+			rowData.Add(SummaryTableElement.Type.ToolMetadata, "framecount", 500.0);
+
+			string tempDir = Path.Combine(Path.GetTempPath(), "PerfReportToolTests_" + Path.GetRandomFileName());
+			Directory.CreateDirectory(tempDir);
+			try
+			{
+				bool writeSuccess = rowData.WriteToCache(tempDir, "test-csv-id");
+				Assert.True(writeSuccess);
+
+				var restored = SummaryTableRowData.TryReadFromCache(tempDir, "test-csv-id");
+				Assert.NotNull(restored);
+				Assert.Equal(rowData.dict.Count, restored.dict.Count);
+
+				foreach (var key in rowData.dict.Keys)
+				{
+					Assert.True(restored.Contains(key));
+					Assert.Equal(rowData.Get(key).name, restored.Get(key).name);
+					Assert.Equal(rowData.Get(key).type, restored.Get(key).type);
+					Assert.Equal(rowData.Get(key).numericValue, restored.Get(key).numericValue);
+					Assert.Equal(rowData.Get(key).value, restored.Get(key).value);
+				}
+			}
+			finally
+			{
+				Directory.Delete(tempDir, true);
+			}
+		}
+
+		[Fact]
+		public void TryReadFromCache_MissingFile_ReturnsNull()
+		{
+			var result = SummaryTableRowData.TryReadFromCache(Path.GetTempPath(), "nonexistent-id");
+			Assert.Null(result);
+		}
+
+		[Fact]
+		[Trait("Category", "FileIO")]
+		public void CacheRoundTrip_ExcludesExternalMetadata()
+		{
+			var rowData = new SummaryTableRowData();
+			rowData.Add(SummaryTableElement.Type.CsvMetadata, "csvid", "test-id");
+			rowData.Add(SummaryTableElement.Type.ExternalMetadata, "external", "value");
+			rowData.Add(SummaryTableElement.Type.CsvStatAverage, "stat", 1.0);
+
+			string tempDir = Path.Combine(Path.GetTempPath(), "PerfReportToolTests_" + Path.GetRandomFileName());
+			Directory.CreateDirectory(tempDir);
+			try
+			{
+				rowData.WriteToCache(tempDir, "test-id");
+				var restored = SummaryTableRowData.TryReadFromCache(tempDir, "test-id");
+
+				Assert.NotNull(restored);
+				// External metadata should not be in the cache
+				Assert.False(restored.Contains("external"));
+				Assert.True(restored.Contains("csvid"));
+				Assert.True(restored.Contains("stat"));
+			}
+			finally
+			{
+				Directory.Delete(tempDir, true);
+			}
+		}
+
+		[Fact]
+		public void ToJsonDict_MetadataOnly_FiltersCorrectly()
+		{
+			var rowData = new SummaryTableRowData();
+			rowData.Add(SummaryTableElement.Type.CsvMetadata, "csvid", "test-id");
+			rowData.Add(SummaryTableElement.Type.CsvMetadata, "platform", "windows");
+			rowData.Add(SummaryTableElement.Type.CsvStatAverage, "frametime", 16.67);
+
+			var jsonDict = rowData.ToJsonDict(bCsvMetadataOnly: true, bWriteAllElementData: false);
+
+			// Should only have CsvMetadata key
+			Assert.Single(jsonDict);
+			Assert.True(jsonDict.ContainsKey("CsvMetadata"));
+			var metaDict = (Dictionary<string, dynamic>)jsonDict["CsvMetadata"];
+			Assert.Equal(2, metaDict.Count);
+		}
+
+		[Fact]
+		public void ToJsonDict_AllTypes_IncludesAllCategories()
+		{
+			var rowData = new SummaryTableRowData();
+			rowData.Add(SummaryTableElement.Type.CsvMetadata, "csvid", "test-id");
+			rowData.Add(SummaryTableElement.Type.CsvStatAverage, "frametime", 16.67);
+			rowData.Add(SummaryTableElement.Type.ToolMetadata, "framecount", 500.0);
+
+			var jsonDict = rowData.ToJsonDict(bCsvMetadataOnly: false, bWriteAllElementData: false);
+
+			// Should have entries for each non-COUNT, non-External type
+			Assert.True(jsonDict.ContainsKey("CsvMetadata"));
+			Assert.True(jsonDict.ContainsKey("CsvStatAverage"));
+			Assert.True(jsonDict.ContainsKey("ToolMetadata"));
+			Assert.False(jsonDict.ContainsKey("ExternalMetadata"));
+		}
+
+		[Fact]
+		public void JsonRoundTrip_PreservesData()
+		{
+			var rowData = new SummaryTableRowData();
+			rowData.Add(SummaryTableElement.Type.CsvMetadata, "csvid", "json-test");
+			rowData.Add(SummaryTableElement.Type.CsvMetadata, "platform", "windows");
+			rowData.Add(SummaryTableElement.Type.CsvStatAverage, "frametime", 16.67);
+			rowData.Add(SummaryTableElement.Type.ToolMetadata, "framecount", 800.0);
+
+			var jsonDict = rowData.ToJsonDict(bCsvMetadataOnly: false, bWriteAllElementData: true);
+			var restored = new SummaryTableRowData(jsonDict);
+
+			Assert.True(restored.Contains("csvid"));
+			Assert.True(restored.Contains("platform"));
+			Assert.True(restored.Contains("frametime"));
+			Assert.True(restored.Contains("framecount"));
+
+			Assert.Equal("json-test", restored.Get("csvid").value);
+			Assert.Equal("windows", restored.Get("platform").value);
+			Assert.Equal(16.67, restored.Get("frametime").numericValue, 2);
+		}
+	}
+}
diff --git a/Engine/Source/Programs/CSVTools/PerfReportTool.Tests/PerfReportToolTests/SummaryTableTests.cs b/Engine/Source/Programs/CSVTools/PerfReportTool.Tests/PerfReportToolTests/SummaryTableTests.cs
new file mode 100644
--- /dev/null
+++ b/Engine/Source/Programs/CSVTools/PerfReportTool.Tests/PerfReportToolTests/SummaryTableTests.cs
@@ -0,0 +1,634 @@
+// Copyright Epic Games, Inc. All Rights Reserved.
+
+using System;
+using System.Collections.Generic;
+using System.IO;
+using System.Linq;
+using PerfReportTool;
+using PerfSummaries;
+using Xunit;
+
+namespace PerfReportTool.Tests.PerfReportToolTests
+{
+	public class SummaryTableColumnTests
+	{
+		#region Construction and basic operations
+
+		[Fact]
+		public void Constructor_SetsNameWithAggregatePrefix()
+		{
+			var col = new SummaryTableColumn("frametime", true, "FrameTime", false, SummaryTableElement.Type.CsvStatAverage,
+				inAggregateType: ColumnAggregateType.Avg);
+			Assert.Equal("Avg frametime", col.name);
+		}
+
+		[Fact]
+		public void Constructor_NoAggregate_NoPrefix()
+		{
+			var col = new SummaryTableColumn("frametime", true, "FrameTime", false, SummaryTableElement.Type.CsvStatAverage);
+			Assert.Equal("frametime", col.name);
+		}
+
+		[Fact]
+		public void GetAggregateTypePrefix_ReturnsCorrectPrefixes()
+		{
+			Assert.Equal("Avg ", SummaryTableColumn.getAggregateTypePrefix(ColumnAggregateType.Avg));
+			Assert.Equal("Min ", SummaryTableColumn.getAggregateTypePrefix(ColumnAggregateType.Min));
+			Assert.Equal("Max ", SummaryTableColumn.getAggregateTypePrefix(ColumnAggregateType.Max));
+			Assert.Equal("", SummaryTableColumn.getAggregateTypePrefix(ColumnAggregateType.None));
+		}
+
+		#endregion
+
+		#region Numeric value operations
+
+		[Fact]
+		public void SetAndGetValue_RoundTrips()
+		{
+			var col = new SummaryTableColumn("stat", true, null, false, SummaryTableElement.Type.CsvStatAverage);
+			col.SetValue(0, 42.5);
+			col.SetValue(1, 99.0);
+			Assert.Equal(42.5, col.GetValue(0));
+			Assert.Equal(99.0, col.GetValue(1));
+		}
+
+		[Fact]
+		public void GetValue_BeyondSize_ReturnsMaxValue()
+		{
+			var col = new SummaryTableColumn("stat", true, null, false, SummaryTableElement.Type.CsvStatAverage);
+			col.SetValue(0, 10.0);
+			Assert.Equal(double.MaxValue, col.GetValue(5));
+		}
+
+		[Fact]
+		public void GetMaxValue_IgnoresMaxValueSentinels()
+		{
+			var col = new SummaryTableColumn("stat", true, null, false, SummaryTableElement.Type.CsvStatAverage);
+			col.SetValue(0, 10.0);
+			col.SetValue(1, 30.0);
+			col.SetValue(2, 20.0);
+			Assert.Equal(30.0, col.GetMaxValue());
+		}
+
+		[Fact]
+		public void GetMinValue_IgnoresMaxValueSentinels()
+		{
+			var col = new SummaryTableColumn("stat", true, null, false, SummaryTableElement.Type.CsvStatAverage);
+			col.SetValue(0, 10.0);
+			col.SetValue(1, 30.0);
+			col.SetValue(2, 20.0);
+			Assert.Equal(10.0, col.GetMinValue());
+		}
+
+		[Fact]
+		public void GetAvgValue_IgnoresMaxValueSentinels()
+		{
+			var col = new SummaryTableColumn("stat", true, null, false, SummaryTableElement.Type.CsvStatAverage);
+			col.SetValue(0, 10.0);
+			col.SetValue(1, 20.0);
+			col.SetValue(2, 30.0);
+			Assert.Equal(20.0, col.GetAvgValue(), 0.001);
+		}
+
+		[Fact]
+		public void GetMaxValue_NoValues_ReturnsZero()
+		{
+			var col = new SummaryTableColumn("stat", true, null, false, SummaryTableElement.Type.CsvStatAverage);
+			Assert.Equal(0.0, col.GetMaxValue());
+		}
+
+		[Fact]
+		public void GetMinValue_NoValues_ReturnsZero()
+		{
+			var col = new SummaryTableColumn("stat", true, null, false, SummaryTableElement.Type.CsvStatAverage);
+			Assert.Equal(0.0, col.GetMinValue());
+		}
+
+		[Fact]
+		public void AreAllValuesOverThreshold_AllOver_ReturnsTrue()
+		{
+			var col = new SummaryTableColumn("stat", true, null, false, SummaryTableElement.Type.CsvStatAverage);
+			col.SetValue(0, 20.0);
+			col.SetValue(1, 30.0);
+			Assert.True(col.AreAllValuesOverThreshold(10.0));
+		}
+
+		[Fact]
+		public void AreAllValuesOverThreshold_SomeBelowButOneAbove_ReturnsFalse()
+		{
+			var col = new SummaryTableColumn("stat", true, null, false, SummaryTableElement.Type.CsvStatAverage);
+			col.SetValue(0, 5.0);
+			col.SetValue(1, 30.0);
+			Assert.False(col.AreAllValuesOverThreshold(10.0));
+		}
+
+		[Fact]
+		public void AreAllValuesOverThreshold_AllBelow_ReturnsFalse()
+		{
+			var col = new SummaryTableColumn("stat", true, null, false, SummaryTableElement.Type.CsvStatAverage);
+			col.SetValue(0, 5.0);
+			col.SetValue(1, 8.0);
+			Assert.False(col.AreAllValuesOverThreshold(10.0));
+		}
+
+		#endregion
+
+		#region String value operations
+
+		[Fact]
+		public void SetAndGetStringValue_RoundTrips()
+		{
+			var col = new SummaryTableColumn("meta", false, null, false, SummaryTableElement.Type.CsvMetadata);
+			col.SetStringValue(0, "hello");
+			col.SetStringValue(1, "world");
+			Assert.Equal("hello", col.GetStringValue(0));
+			Assert.Equal("world", col.GetStringValue(1));
+		}
+
+		[Fact]
+		public void GetCount_ReturnsMaxOfNumericAndString()
+		{
+			var col = new SummaryTableColumn("stat", true, null, false, SummaryTableElement.Type.CsvStatAverage);
+			col.SetValue(0, 10.0);
+			col.SetValue(1, 20.0);
+			col.SetValue(2, 30.0);
+			Assert.Equal(3, col.GetCount());
+		}
+
+		#endregion
+
+		#region GetStatCategory
+
+		[Fact]
+		public void GetStatCategory_CsvStatWithSlash_ReturnsCategory()
+		{
+			var col = new SummaryTableColumn("render/drawcalls", true, null, false, SummaryTableElement.Type.CsvStatAverage);
+			Assert.Equal("render", col.GetStatCategory());
+		}
+
+		[Fact]
+		public void GetStatCategory_CsvStatWithoutSlash_ReturnsEmpty()
+		{
+			var col = new SummaryTableColumn("frametime", true, null, false, SummaryTableElement.Type.CsvStatAverage);
+			Assert.Equal("", col.GetStatCategory());
+		}
+
+		[Fact]
+		public void GetStatCategory_NonCsvStat_ReturnsEmpty()
+		{
+			var col = new SummaryTableColumn("render/drawcalls", true, null, false, SummaryTableElement.Type.SummaryTableMetric);
+			Assert.Equal("", col.GetStatCategory());
+		}
+
+		#endregion
+
+		#region Clone
+
+		[Fact]
+		public void Clone_ProducesIndependentCopy()
+		{
+			var col = new SummaryTableColumn("stat", true, "display", false, SummaryTableElement.Type.CsvStatAverage);
+			col.SetValue(0, 10.0);
+			col.SetValue(1, 20.0);
+
+			var clone = col.Clone();
+			Assert.Equal("stat", clone.name);
+			Assert.Equal(10.0, clone.GetValue(0));
+			Assert.Equal(20.0, clone.GetValue(1));
+
+			// Modify original, verify clone is unaffected
+			col.SetValue(0, 999.0);
+			Assert.Equal(10.0, clone.GetValue(0));
+		}
+
+		#endregion
+	}
+
+	public class SummaryTableStaticTests
+	{
+		#region Element type prefix mapping
+
+		[Fact]
+		public void GetElementTypeStatPrefix_ReturnsCorrectPrefixes()
+		{
+			Assert.Equal("[metric]", SummaryTable.GetElementTypeStatPrefix(SummaryTableElement.Type.SummaryTableMetric));
+			Assert.Equal("[csv]", SummaryTable.GetElementTypeStatPrefix(SummaryTableElement.Type.CsvStatAverage));
+			Assert.Equal("[csvmin]", SummaryTable.GetElementTypeStatPrefix(SummaryTableElement.Type.CsvStatMin));
+			Assert.Equal("[csvmax]", SummaryTable.GetElementTypeStatPrefix(SummaryTableElement.Type.CsvStatMax));
+			Assert.Equal("[meta]", SummaryTable.GetElementTypeStatPrefix(SummaryTableElement.Type.CsvMetadata));
+			Assert.Equal("[toolmeta]", SummaryTable.GetElementTypeStatPrefix(SummaryTableElement.Type.ToolMetadata));
+			Assert.Equal("[extmeta]", SummaryTable.GetElementTypeStatPrefix(SummaryTableElement.Type.ExternalMetadata));
+		}
+
+		#endregion
+
+		#region Qualified stat type parsing
+
+		[Fact]
+		public void GetQualifiedStatType_ParsesCorrectly()
+		{
+			Assert.Equal(SummaryTableElement.Type.SummaryTableMetric, SummaryTable.GetQualfiedStatType("[metric]fps"));
+			Assert.Equal(SummaryTableElement.Type.CsvStatAverage, SummaryTable.GetQualfiedStatType("[csv]frametime"));
+			Assert.Equal(SummaryTableElement.Type.CsvStatMin, SummaryTable.GetQualfiedStatType("[csvmin]frametime"));
+			Assert.Equal(SummaryTableElement.Type.CsvStatMax, SummaryTable.GetQualfiedStatType("[csvmax]frametime"));
+			Assert.Equal(SummaryTableElement.Type.CsvMetadata, SummaryTable.GetQualfiedStatType("[meta]platform"));
+			Assert.Equal(SummaryTableElement.Type.ToolMetadata, SummaryTable.GetQualfiedStatType("[toolmeta]version"));
+			Assert.Equal(SummaryTableElement.Type.ExternalMetadata, SummaryTable.GetQualfiedStatType("[extmeta]extra"));
+		}
+
+		[Fact]
+		public void GetQualifiedStatType_NoPrefix_ReturnsCOUNT()
+		{
+			Assert.Equal(SummaryTableElement.Type.COUNT, SummaryTable.GetQualfiedStatType("frametime"));
+		}
+
+		[Fact]
+		public void GetQualifiedStatType_WithOutputName_StripsPrefix()
+		{
+			SummaryTableElement.Type type = SummaryTable.GetQualfiedStatType("[metric]mvp", out string name);
+			Assert.Equal(SummaryTableElement.Type.SummaryTableMetric, type);
+			Assert.Equal("mvp", name);
+		}
+
+		#endregion
+
+		#region Base stat name parsing
+
+		[Theory]
+		[InlineData("Avg frametime", "frametime")]
+		[InlineData("Min frametime", "frametime")]
+		[InlineData("Max frametime", "frametime")]
+		[InlineData("frametime", "frametime")]
+		public void GetBaseStatName_StripsPrefix(string input, string expected)
+		{
+			Assert.Equal(expected, SummaryTable.GetBaseStatName(input));
+		}
+
+		#endregion
+
+		#region Row visibility with diff rows
+
+		[Fact]
+		public void IsRowVisible_NoDiffRows_AllVisible()
+		{
+			// When bShowOnlyDiffRows=false, ALL rows are visible regardless of diff status
+			Assert.True(SummaryTable.IsRowVisible(DiffRowFrequency.None, false, 0));
+			Assert.True(SummaryTable.IsRowVisible(DiffRowFrequency.None, false, 1));
+			Assert.True(SummaryTable.IsRowVisible(DiffRowFrequency.Alternating, false, 0));
+			Assert.True(SummaryTable.IsRowVisible(DiffRowFrequency.Alternating, false, 2));
+			Assert.True(SummaryTable.IsRowVisible(DiffRowFrequency.Alternating, false, 3));
+		}
+
+		[Fact]
+		public void IsRowVisible_ShowOnlyDiffRows_ShowsDiffRowsOnly()
+		{
+			// Rows 0 and 1 are never diff rows. For Alternating, diff rows start at index 2.
+			// IsDiffRow: row 2 = diff, row 3 = not diff, row 4 = diff
+			Assert.False(SummaryTable.IsRowVisible(DiffRowFrequency.Alternating, true, 0));
+			Assert.False(SummaryTable.IsRowVisible(DiffRowFrequency.Alternating, true, 1));
+			Assert.True(SummaryTable.IsRowVisible(DiffRowFrequency.Alternating, true, 2));  // first diff row
+			Assert.False(SummaryTable.IsRowVisible(DiffRowFrequency.Alternating, true, 3));
+			Assert.True(SummaryTable.IsRowVisible(DiffRowFrequency.Alternating, true, 4));  // second diff row
+		}
+
+		[Fact]
+		public void IsDiffRow_ReturnsCorrectly()
+		{
+			// Rows 0 and 1 are NEVER diff rows
+			Assert.False(SummaryTable.IsDiffRow(DiffRowFrequency.None, 0));
+			Assert.False(SummaryTable.IsDiffRow(DiffRowFrequency.None, 1));
+			Assert.False(SummaryTable.IsDiffRow(DiffRowFrequency.Alternating, 0));
+			Assert.False(SummaryTable.IsDiffRow(DiffRowFrequency.Alternating, 1));
+			// Alternating: (rowIndex - 2) % 2 == 0 for diff rows
+			Assert.True(SummaryTable.IsDiffRow(DiffRowFrequency.Alternating, 2));   // (0 % 2 == 0) = true
+			Assert.False(SummaryTable.IsDiffRow(DiffRowFrequency.Alternating, 3));  // (1 % 2 == 0) = false
+			Assert.True(SummaryTable.IsDiffRow(DiffRowFrequency.Alternating, 4));   // (2 % 2 == 0) = true
+			// AfterEachPair: (rowIndex - 2) % 3 == 0 for diff rows
+			Assert.True(SummaryTable.IsDiffRow(DiffRowFrequency.AfterEachPair, 2));   // (0 % 3 == 0)
+			Assert.False(SummaryTable.IsDiffRow(DiffRowFrequency.AfterEachPair, 3));
+			Assert.False(SummaryTable.IsDiffRow(DiffRowFrequency.AfterEachPair, 4));
+			Assert.True(SummaryTable.IsDiffRow(DiffRowFrequency.AfterEachPair, 5));   // (3 % 3 == 0)
+		}
+
+		#endregion
+	}
+
+	public class SummaryTableDataTests
+	{
+		#region AddRowData and basic table operations
+
+		[Fact]
+		public void AddRowData_IncreasesCount()
+		{
+			var table = new SummaryTable();
+			var row = new SummaryTableRowData();
+			row.Add(SummaryTableElement.Type.CsvStatAverage, "frametime", 16.6);
+			row.Add(SummaryTableElement.Type.CsvMetadata, "platform", "windows");
+
+			table.AddRowData(row, true, false, false, false);
+			Assert.Equal(1, table.Count);
+
+			table.AddRowData(row, true, false, false, false);
+			Assert.Equal(2, table.Count);
+		}
+
+		[Fact]
+		public void AddRowData_CreatesColumns()
+		{
+			var table = new SummaryTable();
+			var row = new SummaryTableRowData();
+			row.Add(SummaryTableElement.Type.CsvStatAverage, "frametime", 16.6);
+			row.Add(SummaryTableElement.Type.CsvMetadata, "platform", "windows");
+
+			table.AddRowData(row, true, false, false, false);
+
+			Assert.NotNull(table.GetColumnByName("frametime"));
+			Assert.NotNull(table.GetColumnByName("platform"));
+		}
+
+		[Fact]
+		public void AddRowData_WithCsvStatMin_IncludesMinColumns()
+		{
+			var table = new SummaryTable();
+			var row = new SummaryTableRowData();
+			row.Add(SummaryTableElement.Type.CsvStatAverage, "frametime", 16.6);
+			row.Add(SummaryTableElement.Type.CsvStatMin, "frametime", 8.0);
+
+			table.AddRowData(row, true, true, false, false);
+
+			// Should have a column for the min value
+			Assert.NotNull(table.GetColumnByName("frametime"));
+		}
+
+		[Fact]
+		public void GetColumnByName_MissingColumn_ReturnsNull()
+		{
+			var table = new SummaryTable();
+			Assert.Null(table.GetColumnByName("nonexistent"));
+		}
+
+		#endregion
+
+		#region SortAndFilter
+
+		[Fact]
+		public void SortAndFilter_FiltersColumns()
+		{
+			var table = new SummaryTable();
+			var row = new SummaryTableRowData();
+			row.Add(SummaryTableElement.Type.CsvStatAverage, "frametime", 16.6);
+			row.Add(SummaryTableElement.Type.CsvStatAverage, "drawcalls", 500.0);
+			row.Add(SummaryTableElement.Type.CsvMetadata, "platform", "windows");
+			table.AddRowData(row, true, false, false, false);
+
+			// Only include frametime in the filter list
+			var filtered = table.SortAndFilter(
+				new List<string> { "frametime" },
+				new List<string>(),
+				false,
+				null,
+				null,
+				false,
+				TableColumnSortMode.Default);
+
+			Assert.NotNull(filtered.GetColumnByName("frametime"));
+			// drawcalls should be filtered out
+			Assert.Null(filtered.GetColumnByName("drawcalls"));
+		}
+
+		[Fact]
+		public void SortAndFilter_WildcardFilter_MatchesMultipleColumns()
+		{
+			var table = new SummaryTable();
+			var row = new SummaryTableRowData();
+			row.Add(SummaryTableElement.Type.CsvStatAverage, "render/drawcalls", 500.0);
+			row.Add(SummaryTableElement.Type.CsvStatAverage, "render/triangles", 1000.0);
+			row.Add(SummaryTableElement.Type.CsvStatAverage, "physics/time", 5.0);
+			table.AddRowData(row, true, false, false, false);
+
+			var filtered = table.SortAndFilter(
+				new List<string> { "render/*" },
+				new List<string>(),
+				false,
+				null,
+				null,
+				false,
+				TableColumnSortMode.Default);
+
+			Assert.NotNull(filtered.GetColumnByName("render/drawcalls"));
+			Assert.NotNull(filtered.GetColumnByName("render/triangles"));
+			Assert.Null(filtered.GetColumnByName("physics/time"));
+		}
+
+		[Fact]
+		public void SortAndFilter_WithAdditionalFilters_AppliesFilters()
+		{
+			var table = new SummaryTable();
+			var row = new SummaryTableRowData();
+			row.Add(SummaryTableElement.Type.CsvStatAverage, "lowstat", 0.001);
+			row.Add(SummaryTableElement.Type.CsvStatAverage, "highstat", 100.0);
+			table.AddRowData(row, true, false, false, false);
+
+			var filters = new List<ISummaryTableColumnFilter>
+			{
+				new StatThresholdColumnFilter(1.0f)
+			};
+
+			var filtered = table.SortAndFilter(
+				new List<string> { "lowstat", "highstat" },
+				new List<string>(),
+				false,
+				null,
+				filters,
+				false,
+				TableColumnSortMode.Default);
+
+			// highstat should remain, lowstat should be filtered
+			Assert.NotNull(filtered.GetColumnByName("highstat"));
+			Assert.Null(filtered.GetColumnByName("lowstat"));
+		}
+
+		#endregion
+
+		#region CollateSortedTable
+
+		[Fact]
+		public void CollateSortedTable_MergesRowsByKey()
+		{
+			var table = new SummaryTable();
+
+			// Add two rows with same platform
+			var row1 = new SummaryTableRowData();
+			row1.Add(SummaryTableElement.Type.CsvMetadata, "platform", "windows");
+			row1.Add(SummaryTableElement.Type.CsvStatAverage, "frametime", 16.0);
+			table.AddRowData(row1, true, false, false, false);
+
+			var row2 = new SummaryTableRowData();
+			row2.Add(SummaryTableElement.Type.CsvMetadata, "platform", "windows");
+			row2.Add(SummaryTableElement.Type.CsvStatAverage, "frametime", 20.0);
+			table.AddRowData(row2, true, false, false, false);
+
+			// Add one row with different platform
+			var row3 = new SummaryTableRowData();
+			row3.Add(SummaryTableElement.Type.CsvMetadata, "platform", "linux");
+			row3.Add(SummaryTableElement.Type.CsvStatAverage, "frametime", 33.0);
+			table.AddRowData(row3, true, false, false, false);
+
+			// First sort by platform so same platforms are adjacent
+			var sorted = table.SortRows(new List<string> { "platform" }, false);
+			var collated = sorted.CollateSortedTable(
+				new List<string> { "platform" },
+				false,
+				StringCollationVisibility.Auto,
+				DateCollationVisibility.Newest);
+
+			// Should have 2 collated rows (linux and windows)
+			Assert.Equal(2, collated.Count);
+		}
+
+		[Fact]
+		public void CollateSortedTable_WithMinMaxColumns_AddsAggregates()
+		{
+			var table = new SummaryTable();
+
+			var row1 = new SummaryTableRowData();
+			row1.Add(SummaryTableElement.Type.CsvMetadata, "platform", "windows");
+			row1.Add(SummaryTableElement.Type.CsvStatAverage, "frametime", 16.0);
+			table.AddRowData(row1, true, false, false, false);
+
+			var row2 = new SummaryTableRowData();
+			row2.Add(SummaryTableElement.Type.CsvMetadata, "platform", "windows");
+			row2.Add(SummaryTableElement.Type.CsvStatAverage, "frametime", 20.0);
+			table.AddRowData(row2, true, false, false, false);
+
+			var sorted = table.SortRows(new List<string> { "platform" }, false);
+			var collated = sorted.CollateSortedTable(
+				new List<string> { "platform" },
+				true, // addMinMaxColumns
+				StringCollationVisibility.Auto,
+				DateCollationVisibility.Newest);
+
+			// Collated table should have at least 1 row and contain aggregate columns
+			Assert.Equal(1, collated.Count);
+
+			// The collated column names include aggregate prefixes: "Avg frametime", "Min frametime", "Max frametime"
+			// GetColumnByName may use the full name or lookup key. Verify we have more than just the collateBy columns.
+			// We can enumerate to verify aggregate columns exist
+			var rows = collated.EnumerateTableDataByRow().ToList();
+			Assert.True(rows.Count > 0, "Collated table should have rows");
+		}
+
+		#endregion
+
+		#region WriteToCSV
+
+		[Fact]
+		[Trait("Category", "FileIO")]
+		public void WriteToCSV_ProducesValidFile()
+		{
+			var table = new SummaryTable();
+			var row = new SummaryTableRowData();
+			row.Add(SummaryTableElement.Type.CsvMetadata, "platform", "windows");
+			row.Add(SummaryTableElement.Type.CsvStatAverage, "frametime", 16.6);
+			table.AddRowData(row, true, false, false, false);
+
+			string tempFile = Path.Combine(Path.GetTempPath(), "test_summary_" + Path.GetRandomFileName() + ".csv");
+			try
+			{
+				table.WriteToCSV(tempFile);
+				Assert.True(File.Exists(tempFile));
+				string content = File.ReadAllText(tempFile);
+				Assert.Contains("frametime", content);
+			}
+			finally
+			{
+				if (File.Exists(tempFile)) File.Delete(tempFile);
+			}
+		}
+
+		#endregion
+
+		#region SortRows
+
+		[Fact]
+		public void SortRows_SortsBySpecifiedColumn()
+		{
+			var table = new SummaryTable();
+
+			var row1 = new SummaryTableRowData();
+			row1.Add(SummaryTableElement.Type.CsvMetadata, "platform", "windows");
+			row1.Add(SummaryTableElement.Type.CsvStatAverage, "frametime", 33.0);
+			table.AddRowData(row1, true, false, false, false);
+
+			var row2 = new SummaryTableRowData();
+			row2.Add(SummaryTableElement.Type.CsvMetadata, "platform", "android");
+			row2.Add(SummaryTableElement.Type.CsvStatAverage, "frametime", 16.0);
+			table.AddRowData(row2, true, false, false, false);
+
+			var row3 = new SummaryTableRowData();
+			row3.Add(SummaryTableElement.Type.CsvMetadata, "platform", "linux");
+			row3.Add(SummaryTableElement.Type.CsvStatAverage, "frametime", 22.0);
+			table.AddRowData(row3, true, false, false, false);
+
+			var sorted = table.SortRows(new List<string> { "platform" }, false);
+
+			var platformCol = sorted.GetColumnByName("platform");
+			Assert.NotNull(platformCol);
+
+			// Should be sorted alphabetically: android, linux, windows
+			Assert.Equal("android", platformCol.GetStringValue(0));
+			Assert.Equal("linux", platformCol.GetStringValue(1));
+			Assert.Equal("windows", platformCol.GetStringValue(2));
+		}
+
+		[Fact]
+		public void SortRows_ReverseSort_ReversesOrder()
+		{
+			var table = new SummaryTable();
+
+			var row1 = new SummaryTableRowData();
+			row1.Add(SummaryTableElement.Type.CsvMetadata, "platform", "a");
+			table.AddRowData(row1, true, false, false, false);
+
+			var row2 = new SummaryTableRowData();
+			row2.Add(SummaryTableElement.Type.CsvMetadata, "platform", "b");
+			table.AddRowData(row2, true, false, false, false);
+
+			var sorted = table.SortRows(new List<string> { "platform" }, true);
+			var platformCol = sorted.GetColumnByName("platform");
+			Assert.Equal("b", platformCol.GetStringValue(0));
+			Assert.Equal("a", platformCol.GetStringValue(1));
+		}
+
+		#endregion
+	}
+
+	public class TableUtilTests
+	{
+		[Fact]
+		public void FormatStatName_AddsWordBreakAfterSlash()
+		{
+			Assert.Equal("render/<wbr>drawcalls", TableUtil.FormatStatName("render/drawcalls"));
+		}
+
+		[Fact]
+		public void SanitizeHtmlString_EscapesAngleBrackets()
+		{
+			Assert.Equal("&lt;script&gt;", TableUtil.SanitizeHtmlString("<script>"));
+		}
+
+		[Fact]
+		public void SafeTruncateHtmlTableValue_TruncatesLongString()
+		{
+			string result = TableUtil.SafeTruncateHtmlTableValue("This is a very long string", 10);
+			Assert.EndsWith("...", result);
+		}
+
+		[Fact]
+		public void SafeTruncateHtmlTableValue_LinkTemplate_NotTruncated()
+		{
+			string input = "{LinkTemplate:Report:csvid123}";
+			Assert.Equal(input, TableUtil.SafeTruncateHtmlTableValue(input, 5));
+		}
+	}
+}
diff --git a/Engine/Source/Programs/CSVTools/PerfReportTool.Tests/PerfReportToolTests/SummaryTypeTests.cs b/Engine/Source/Programs/CSVTools/PerfReportTool.Tests/PerfReportToolTests/SummaryTypeTests.cs
new file mode 100644
--- /dev/null
+++ b/Engine/Source/Programs/CSVTools/PerfReportTool.Tests/PerfReportToolTests/SummaryTypeTests.cs
@@ -0,0 +1,309 @@
+// Copyright Epic Games, Inc. All Rights Reserved.
+
+using System;
+using System.Collections.Generic;
+using System.Globalization;
+using System.Linq;
+using System.Xml.Linq;
+using CSVStats;
+using PerfReportTool;
+using PerfSummaries;
+using PerfReportTool.Tests.TestHelpers;
+using Xunit;
+
+namespace PerfReportTool.Tests.PerfReportToolTests
+{
+	public class StandardSummaryTests
+	{
+		[Fact]
+		public void GetName_ReturnsStandard()
+		{
+			var summary = new StandardSummary();
+			Assert.Equal("standard", summary.GetName());
+		}
+
+		[Fact]
+		public void WriteSummaryData_ComputesTotalTime()
+		{
+			var summary = new StandardSummary();
+			// 1000 frames at 16.67ms each = 16670ms = 16.67s
+			var csvStats = CsvTestHelper.CreateSimpleCsvStats(new[] { "frametime" }, 1000, 16.67f);
+			var rowData = new SummaryTableRowData();
+
+			summary.WriteSummaryData(false, csvStats, csvStats, false, rowData, null);
+
+			Assert.True(rowData.Contains("total time (s)"));
+			double totalTime = rowData.Get("total time (s)").numericValue;
+			// 1000 * 16.67ms / 1000.0 = 16.67s
+			Assert.Equal(16.67, totalTime, 0.01);
+		}
+
+		[Fact]
+		public void WriteSummaryData_DoesNotOverwriteExisting()
+		{
+			var summary = new StandardSummary();
+			var csvStats = CsvTestHelper.CreateSimpleCsvStats(new[] { "frametime" }, 100, 10.0f);
+			var rowData = new SummaryTableRowData();
+			rowData.Add(SummaryTableElement.Type.SummaryTableMetric, "Total Time (s)", 999.0);
+
+			summary.WriteSummaryData(false, csvStats, csvStats, false, rowData, null);
+
+			// Should not overwrite the existing value
+			Assert.Equal(999.0, rowData.Get("total time (s)").numericValue);
+		}
+
+		[Fact]
+		public void WriteSummaryData_NullRowData_DoesNotThrow()
+		{
+			var summary = new StandardSummary();
+			var csvStats = CsvTestHelper.CreateSimpleCsvStats(new[] { "frametime" }, 100, 10.0f);
+			// Should not throw with null rowData
+			summary.WriteSummaryData(false, csvStats, csvStats, false, null, null);
+		}
+	}
+
+	public class FPSChartSummaryTests
+	{
+		private static FPSChartSummary CreateFPSChartSummary(int fps = 30, float hitchThreshold = 150.0f, bool useEngineHitchMetric = false, int recommendedFps = 0, float slowFrameTimeBuffer = 0)
+		{
+			string ht = hitchThreshold.ToString(CultureInfo.InvariantCulture);
+			string sfb = slowFrameTimeBuffer.ToString(CultureInfo.InvariantCulture);
+			string xml = $@"<summary type='fpschart' fps='{fps}' hitchThreshold='{ht}' useEngineHitchMetric='{useEngineHitchMetric.ToString().ToLower()}' recommendedFps='{recommendedFps}' slowFrameTimeBuffer='{sfb}'>
+				<stats/>
+			</summary>";
+			var element = XElement.Parse(xml);
+			var vars = new XmlVariableMappings();
+			return new FPSChartSummary(element, vars, "");
+		}
+
+		[Fact]
+		public void GetName_ReturnsFpschart()
+		{
+			var summary = new FPSChartSummary();
+			Assert.Equal("fpschart", summary.GetName());
+		}
+
+		[Fact]
+		public void WriteSummaryData_ComputesMVP()
+		{
+			var summary = CreateFPSChartSummary(fps: 30, hitchThreshold: 150.0f);
+			// Create 30fps-equivalent frameTimes (33.33ms each) for 10 seconds (300 frames)
+			var csvStats = CsvTestHelper.CreateSimpleCsvStats(new[] { "frametime" }, 300, 33.33f);
+			var rowData = new SummaryTableRowData();
+
+			summary.WriteSummaryData(false, csvStats, csvStats, false, rowData, null);
+
+			Assert.True(rowData.Contains("mvp"));
+			double mvp = rowData.Get("mvp").numericValue;
+			// At exact 30fps, MVP should be ~0 (no missed frames)
+			Assert.True(mvp < 5.0, $"MVP should be near 0 for exact target frame rate, got {mvp}");
+		}
+
+		[Fact]
+		public void WriteSummaryData_ComputesHitchesPerMinute()
+		{
+			var summary = CreateFPSChartSummary(fps: 30, hitchThreshold: 50.0f);
+			// Create frames: mostly 33ms, but every 30th frame is a 100ms hitch
+			float[] frameTimes = new float[300];
+			for (int i = 0; i < 300; i++)
+			{
+				frameTimes[i] = (i % 30 == 0) ? 100.0f : 33.33f;
+			}
+			var csvStats = new CsvStats();
+			var stat = CsvTestHelper.CreateStatSamples("frametime", frameTimes);
+			csvStats.AddStat(stat);
+			var rowData = new SummaryTableRowData();
+
+			summary.WriteSummaryData(false, csvStats, csvStats, false, rowData, null);
+
+			Assert.True(rowData.Contains("hitches/min"));
+			double hitchesPerMin = rowData.Get("hitches/min").numericValue;
+			Assert.True(hitchesPerMin > 0, "Should detect hitches");
+		}
+
+		[Fact]
+		public void WriteSummaryData_ComputesTargetFPS()
+		{
+			var summary = CreateFPSChartSummary(fps: 60, hitchThreshold: 150.0f);
+			var csvStats = CsvTestHelper.CreateSimpleCsvStats(new[] { "frametime" }, 100, 16.67f);
+			var rowData = new SummaryTableRowData();
+
+			summary.WriteSummaryData(false, csvStats, csvStats, false, rowData, null);
+
+			Assert.True(rowData.Contains("targetfps"));
+			Assert.Equal(60.0, rowData.Get("targetfps").numericValue);
+		}
+
+		[Fact]
+		public void WriteSummaryData_ComputesTotalTime()
+		{
+			var summary = CreateFPSChartSummary(fps: 30, hitchThreshold: 150.0f);
+			// 100 frames at 33.33ms = ~3.33 seconds
+			var csvStats = CsvTestHelper.CreateSimpleCsvStats(new[] { "frametime" }, 100, 33.33f);
+			var rowData = new SummaryTableRowData();
+
+			summary.WriteSummaryData(false, csvStats, csvStats, false, rowData, null);
+
+			Assert.True(rowData.Contains("total time (s)"));
+			double totalTime = rowData.Get("total time (s)").numericValue;
+			// 99 frames counted (last frame skipped) * 33.33ms / 1000
+			Assert.True(totalTime > 3.0 && totalTime < 4.0, $"Expected ~3.3s, got {totalTime}");
+		}
+
+		[Fact]
+		public void WriteSummaryData_NoHitches_ZeroHitchTimePercent()
+		{
+			var summary = CreateFPSChartSummary(fps: 30, hitchThreshold: 150.0f);
+			// All frames at 33ms - well below 150ms hitch threshold
+			var csvStats = CsvTestHelper.CreateSimpleCsvStats(new[] { "frametime" }, 100, 33.33f);
+			var rowData = new SummaryTableRowData();
+
+			summary.WriteSummaryData(false, csvStats, csvStats, false, rowData, null);
+
+			Assert.True(rowData.Contains("hitchtimepercent"));
+			double htp = rowData.Get("hitchtimepercent").numericValue;
+			Assert.Equal(0.0, htp, 0.01);
+		}
+
+		[Fact]
+		public void WriteSummaryData_SlowFrameTimePercent_WithRecommendedFps()
+		{
+			var summary = CreateFPSChartSummary(fps: 30, hitchThreshold: 150.0f, recommendedFps: 30, slowFrameTimeBuffer: 0);
+			// Half frames fast, half slow
+			float[] frameTimes = new float[100];
+			for (int i = 0; i < 100; i++)
+			{
+				frameTimes[i] = (i < 50) ? 20.0f : 50.0f; // 50ms >> 33.33ms target
+			}
+			var csvStats = new CsvStats();
+			var stat = CsvTestHelper.CreateStatSamples("frametime", frameTimes);
+			csvStats.AddStat(stat);
+			var rowData = new SummaryTableRowData();
+
+			summary.WriteSummaryData(false, csvStats, csvStats, false, rowData, null);
+
+			Assert.True(rowData.Contains("slowframetimepercent"));
+			double sftp = rowData.Get("slowframetimepercent").numericValue;
+			Assert.True(sftp > 0, "Should have some slow frame time percent");
+		}
+
+		[Fact]
+		public void WriteSummaryData_EngineHitchMetric_CountsHitchesDifferently()
+		{
+			var summarySimple = CreateFPSChartSummary(fps: 30, hitchThreshold: 50.0f, useEngineHitchMetric: false);
+			var summaryEngine = CreateFPSChartSummary(fps: 30, hitchThreshold: 50.0f, useEngineHitchMetric: true);
+
+			// Create frames with a slow section
+			float[] frameTimes = new float[100];
+			for (int i = 0; i < 100; i++)
+			{
+				frameTimes[i] = (i >= 10 && i <= 15) ? 80.0f : 33.33f;
+			}
+			var csvStats = new CsvStats();
+			csvStats.AddStat(CsvTestHelper.CreateStatSamples("frametime", frameTimes));
+
+			var rowDataSimple = new SummaryTableRowData();
+			var rowDataEngine = new SummaryTableRowData();
+
+			summarySimple.WriteSummaryData(false, csvStats, csvStats, false, rowDataSimple, null);
+			summaryEngine.WriteSummaryData(false, csvStats, csvStats, false, rowDataEngine, null);
+
+			double hitchesSimple = rowDataSimple.Get("hitches/min").numericValue;
+			double hitchesEngine = rowDataEngine.Get("hitches/min").numericValue;
+
+			// Simple mode should detect the 80ms frames as hitches (above 50ms threshold)
+			Assert.True(hitchesSimple > 0, $"Simple hitch counting should detect hitches, got {hitchesSimple}");
+			// The two modes should produce different hitch counts
+			Assert.NotEqual(hitchesSimple, hitchesEngine);
+		}
+	}
+
+	public class HitchSummaryTests
+	{
+		private static HitchSummary CreateHitchSummary(string thresholds = "50,100,150,200")
+		{
+			string xml = $@"<summary type='hitches'>
+				<hitchThresholds>{thresholds}</hitchThresholds>
+				<stats>frametime</stats>
+			</summary>";
+			var element = XElement.Parse(xml);
+			var vars = new XmlVariableMappings();
+			return new HitchSummary(element, vars, "");
+		}
+
+		[Fact]
+		public void GetName_ReturnsHitches()
+		{
+			var summary = new HitchSummary();
+			Assert.Equal("hitches", summary.GetName());
+		}
+
+		[Fact]
+		public void Constructor_ParsesThresholds()
+		{
+			var summary = CreateHitchSummary("50,100,150,200");
+			Assert.Equal(4, summary.HitchThresholds.Length);
+			Assert.Equal(50.0, summary.HitchThresholds[0]);
+			Assert.Equal(200.0, summary.HitchThresholds[3]);
+		}
+
+		[Fact]
+		public void WriteSummaryData_AddsHitchMetrics_100msIntervals()
+		{
+			var summary = CreateHitchSummary();
+			// Some frames with hitches
+			float[] frameTimes = new float[100];
+			for (int i = 0; i < 100; i++)
+			{
+				frameTimes[i] = (i == 50) ? 500.0f : 16.67f; // One 500ms hitch
+			}
+			var csvStats = new CsvStats();
+			csvStats.AddStat(CsvTestHelper.CreateStatSamples("frametime", frameTimes));
+			var rowData = new SummaryTableRowData();
+
+			summary.WriteSummaryData(false, csvStats, csvStats, false, rowData, null);
+
+			// Should have Hitches>100ms through Hitches>1000ms entries
+			Assert.True(rowData.Contains("hitches>100ms"));
+			Assert.True(rowData.Contains("hitches>500ms"));
+			Assert.True(rowData.Contains("hitches>1000ms"));
+
+			// A 500ms frame is strictly >100ms through >400ms, but not >500ms (strict greater-than)
+			double hitches100 = rowData.Get("hitches>100ms").numericValue;
+			double hitches500 = rowData.Get("hitches>500ms").numericValue;
+			double hitches600 = rowData.Get("hitches>600ms").numericValue;
+
+			Assert.True(hitches100 >= 1, "Should have at least 1 hitch >100ms");
+			Assert.Equal(0.0, hitches500);
+			Assert.Equal(0.0, hitches600);
+		}
+
+		[Fact]
+		public void WriteSummaryData_NoHitches_AllZero()
+		{
+			var summary = CreateHitchSummary();
+			var csvStats = CsvTestHelper.CreateSimpleCsvStats(new[] { "frametime" }, 100, 16.67f);
+			var rowData = new SummaryTableRowData();
+
+			summary.WriteSummaryData(false, csvStats, csvStats, false, rowData, null);
+
+			for (int i = 1; i <= 10; i++)
+			{
+				string key = $"hitches>{i * 100}ms";
+				Assert.True(rowData.Contains(key), $"Missing {key}");
+				Assert.Equal(0.0, rowData.Get(key).numericValue);
+			}
+		}
+	}
+
+	public class SummaryBaseClassTests
+	{
+		[Fact]
+		public void StandardSummary_CanBeInstantiated()
+		{
+			var summary = new StandardSummary();
+			Assert.Equal("standard", summary.GetName());
+		}
+	}
+
+}
diff --git a/Engine/Source/Programs/CSVTools/PerfReportTool.Tests/PerfReportToolTests/XmlVariableMappingsTests.cs b/Engine/Source/Programs/CSVTools/PerfReportTool.Tests/PerfReportToolTests/XmlVariableMappingsTests.cs
new file mode 100644
--- /dev/null
+++ b/Engine/Source/Programs/CSVTools/PerfReportTool.Tests/PerfReportToolTests/XmlVariableMappingsTests.cs
@@ -0,0 +1,134 @@
+using PerfReportTool;
+using PerfReportTool.Tests.TestHelpers;
+using Xunit;
+
+namespace PerfReportTool.Tests.PerfReportToolTests
+{
+	public class XmlVariableMappingsTests
+	{
+		[Fact]
+		public void SetVariable_ValidName_StoresCorrectly()
+		{
+			var vars = new XmlVariableMappings();
+			vars.SetVariable("myVar", "myValue");
+
+			string result = vars.ResolveVariables("${myVar}");
+			Assert.Equal("myValue", result);
+		}
+
+		[Fact]
+		public void SetVariable_WithDot_StoresCorrectly()
+		{
+			var vars = new XmlVariableMappings();
+			vars.SetVariable("my.var", "dotValue");
+
+			string result = vars.ResolveVariables("${my.var}");
+			Assert.Equal("dotValue", result);
+		}
+
+		[Fact]
+		public void SetVariable_WithUnderscore_StoresCorrectly()
+		{
+			var vars = new XmlVariableMappings();
+			vars.SetVariable("my_var", "underscoreValue");
+
+			string result = vars.ResolveVariables("${my_var}");
+			Assert.Equal("underscoreValue", result);
+		}
+
+		[Fact]
+		public void SetVariable_InvalidName_ThrowsException()
+		{
+			var vars = new XmlVariableMappings();
+			Assert.Throws<System.Exception>(() => vars.SetVariable("invalid name!", "value"));
+		}
+
+		[Fact]
+		public void SetMetadataVariables_SetsMetaPrefix()
+		{
+			var vars = new XmlVariableMappings();
+			var metadata = CsvTestHelper.CreateMetadata(("platform", "windows"), ("config", "test"));
+			vars.SetMetadataVariables(metadata);
+
+			Assert.Equal("windows", vars.ResolveVariables("${meta.platform}"));
+			Assert.Equal("test", vars.ResolveVariables("${meta.config}"));
+		}
+
+		[Fact]
+		public void ResolveVariables_NoVariables_ReturnsUnchanged()
+		{
+			var vars = new XmlVariableMappings();
+			string input = "no variables here";
+			Assert.Equal(input, vars.ResolveVariables(input));
+		}
+
+		[Fact]
+		public void ResolveVariables_SimpleVariable_Resolves()
+		{
+			var vars = new XmlVariableMappings();
+			vars.SetVariable("name", "World");
+
+			string result = vars.ResolveVariables("Hello ${name}!");
+			Assert.Equal("Hello World!", result);
+		}
+
+		[Fact]
+		public void ResolveVariables_MultipleVariables_ResolvesAll()
+		{
+			var vars = new XmlVariableMappings();
+			vars.SetVariable("a", "1");
+			vars.SetVariable("b", "2");
+
+			string result = vars.ResolveVariables("${a} and ${b}");
+			Assert.Equal("1 and 2", result);
+		}
+
+		[Fact]
+		public void ResolveVariables_MissingVariable_ReturnsEmptyString()
+		{
+			var vars = new XmlVariableMappings();
+			string result = vars.ResolveVariables("before ${missing} after");
+			Assert.Equal("before  after", result);
+		}
+
+		[Fact]
+		public void ResolveVariables_ArrayIndexVariable_ResolvesElement()
+		{
+			var vars = new XmlVariableMappings();
+			vars.SetVariable("colors", "red,green,blue");
+
+			Assert.Equal("red", vars.ResolveVariables("${colors[0]}"));
+			Assert.Equal("green", vars.ResolveVariables("${colors[1]}"));
+			Assert.Equal("blue", vars.ResolveVariables("${colors[2]}"));
+		}
+
+		[Fact]
+		public void ResolveVariables_OverwrittenVariable_UsesLatestValue()
+		{
+			var vars = new XmlVariableMappings();
+			vars.SetVariable("x", "first");
+			Assert.Equal("first", vars.ResolveVariables("${x}"));
+
+			vars.SetVariable("x", "second");
+			Assert.Equal("second", vars.ResolveVariables("${x}"));
+		}
+
+		[Fact]
+		public void ResolveVariables_NoDollarSign_SkipsProcessing()
+		{
+			var vars = new XmlVariableMappings();
+			// No $ means fast path - no processing
+			string result = vars.ResolveVariables("simple text");
+			Assert.Equal("simple text", result);
+		}
+
+		[Fact]
+		public void ResolveVariables_IncompleteSyntax_ReturnsPartially()
+		{
+			var vars = new XmlVariableMappings();
+			// Incomplete variable syntax - ${... without closing brace
+			string result = vars.ResolveVariables("${unclosed");
+			Assert.Equal("${unclosed", result);
+		}
+	}
+}
diff --git a/Engine/Source/Programs/CSVTools/PerfReportTool.Tests/TestDataFixture.cs b/Engine/Source/Programs/CSVTools/PerfReportTool.Tests/TestDataFixture.cs
new file mode 100644
--- /dev/null
+++ b/Engine/Source/Programs/CSVTools/PerfReportTool.Tests/TestDataFixture.cs
@@ -0,0 +1,82 @@
+using System;
+using System.IO;
+
+namespace PerfReportTool.Tests
+{
+	public class TestDataFixture : IDisposable
+	{
+		/// <summary>Path to the in-project TestData directory (always available).</summary>
+		public string TestDataDir { get; }
+		public bool IsAvailable { get; }
+
+		public TestDataFixture()
+		{
+			// 1. Check environment variable override (supports absolute or relative paths)
+			string envPath = Environment.GetEnvironmentVariable("PERFREPORT_TEST_DATA_DIR");
+			if (!string.IsNullOrEmpty(envPath))
+			{
+				string resolved = Path.GetFullPath(envPath);
+				if (Directory.Exists(resolved))
+				{
+					TestDataDir = resolved;
+					IsAvailable = true;
+					return;
+				}
+			}
+
+			// 2. Find the in-project TestData directory by walking up from the output directory
+			string dir = AppContext.BaseDirectory;
+			for (int i = 0; i < 10; i++)
+			{
+				string candidate = Path.Combine(dir, "TestData");
+				if (Directory.Exists(candidate))
+				{
+					TestDataDir = candidate;
+					IsAvailable = true;
+					return;
+				}
+				string parent = Path.GetDirectoryName(dir);
+				if (parent == null || parent == dir)
+				{
+					break;
+				}
+				dir = parent;
+			}
+
+			TestDataDir = null;
+			IsAvailable = false;
+		}
+
+		public string GetFile(string relativePath)
+		{
+			if (TestDataDir == null)
+			{
+				throw new InvalidOperationException(
+					"TestDataDir is not available. Call RequireTestData() before GetFile().");
+			}
+			return Path.Combine(TestDataDir, relativePath);
+		}
+
+		public string GetPerfCsvDir()
+		{
+			return TestDataDir;
+		}
+
+		/// <summary>
+		/// Call at the start of any test that requires test data files.
+		/// </summary>
+		public void RequireTestData()
+		{
+			if (!IsAvailable)
+			{
+				throw new InvalidOperationException(
+					"Test data not available. Ensure the TestData directory exists in the test project, " +
+					"or set PERFREPORT_TEST_DATA_DIR environment variable.");
+			}
+		}
+
+		public void Dispose()
+		{
+		}
+	}
+}
diff --git a/Engine/Source/Programs/CSVTools/PerfReportTool.Tests/TestHelpers/CsvTestHelper.cs b/Engine/Source/Programs/CSVTools/PerfReportTool.Tests/TestHelpers/CsvTestHelper.cs
new file mode 100644
--- /dev/null
+++ b/Engine/Source/Programs/CSVTools/PerfReportTool.Tests/TestHelpers/CsvTestHelper.cs
@@ -0,0 +1,48 @@
+using System.Collections.Generic;
+using CSVStats;
+using PerfSummaries;
+
+namespace PerfReportTool.Tests.TestHelpers
+{
+	internal static class CsvTestHelper
+	{
+		public static CSVStats.CsvStats CreateSimpleCsvStats(string[] statNames, int sampleCount, float defaultValue = 0.0f)
+		{
+			var csvStats = new CSVStats.CsvStats();
+			foreach (string name in statNames)
+			{
+				var stat = new StatSamples(name);
+				for (int i = 0; i < sampleCount; i++)
+				{
+					stat.samples.Add(defaultValue);
+				}
+				stat.ComputeAverageAndTotal();
+				csvStats.AddStat(stat);
+			}
+			return csvStats;
+		}
+
+		public static CsvMetadata CreateMetadata(params (string key, string value)[] pairs)
+		{
+			var metadata = new CsvMetadata();
+			foreach (var (key, value) in pairs)
+			{
+				metadata.Values[key] = value;
+			}
+			return metadata;
+		}
+
+		public static StatSamples CreateStatSamples(string name, float[] values)
+		{
+			var stat = new StatSamples(name);
+			stat.samples.AddRange(values);
+			stat.ComputeAverageAndTotal();
+			return stat;
+		}
+
+		public static ColourThresholdList CreateDefaultThresholds(double red, double orange, double yellow, double green)
+		{
+			return new ColourThresholdList(red, orange, yellow, green);
+		}
+	}
+}
diff --git a/Engine/Source/Programs/CSVTools/PerfReportTool.Tests/TestHelpers/ReportXmlTestHelper.cs b/Engine/Source/Programs/CSVTools/PerfReportTool.Tests/TestHelpers/ReportXmlTestHelper.cs
new file mode 100644
--- /dev/null
+++ b/Engine/Source/Programs/CSVTools/PerfReportTool.Tests/TestHelpers/ReportXmlTestHelper.cs
@@ -0,0 +1,38 @@
+using System.IO;
+using System.Reflection;
+
+namespace PerfReportTool.Tests.TestHelpers
+{
+	internal static class ReportXmlTestHelper
+	{
+		/// <summary>
+		/// Gets the directory containing the PerfReportTool output (where ReportTypes.xml etc. are copied).
+		/// </summary>
+		public static string GetPerfReportToolOutputDir()
+		{
+			// The PerfReportTool assembly is loaded as a dependency. Find its location.
+			string assemblyLocation = typeof(PerfReportTool.Version).Assembly.Location;
+			return Path.GetDirectoryName(assemblyLocation);
+		}
+
+		public static string GetReportTypesXmlPath()
+		{
+			return Path.Combine(GetPerfReportToolOutputDir(), "ReportTypes.xml");
+		}
+
+		public static string GetReportGraphsXmlPath()
+		{
+			return Path.Combine(GetPerfReportToolOutputDir(), "ReportGraphs.xml");
+		}
+
+		public static string GetLlmReportTypesXmlPath()
+		{
+			return Path.Combine(GetPerfReportToolOutputDir(), "LLMReportTypes.xml");
+		}
+
+		public static string GetLlmReportGraphsXmlPath()
+		{
+			return Path.Combine(GetPerfReportToolOutputDir(), "LLMReportGraphs.xml");
+		}
+	}
+}
diff --git a/Engine/Source/Programs/CSVTools/PerfReportTool/Properties/AssemblyInfo.cs b/Engine/Source/Programs/CSVTools/PerfReportTool/Properties/AssemblyInfo.cs
--- a/Engine/Source/Programs/CSVTools/PerfReportTool/Properties/AssemblyInfo.cs
+++ b/Engine/Source/Programs/CSVTools/PerfReportTool/Properties/AssemblyInfo.cs
@@ -5,3 +5,4 @@ using System.Runtime.InteropServices;
 
 // The following GUID is for the ID of the typelib if this project is exposed to COM
 [assembly: Guid("ae427314-972a-4178-8c36-fc3efcd164e3")]
+[assembly: System.Runtime.CompilerServices.InternalsVisibleTo("PerfReportTool.Tests")]
