PRT - Add options to write out summary tables as json such that it better reflects the HTML output.

+ Adds -summaryTableOutputFormats which accepts one or more arguments (html,csv,json) to output summary tables in these formats.
+ The json arg works differently from -summaryTableToJson as it will output the filtered summary table rather than the full csvs. This also supports other summary table args like collation, rowsorts, custom table etc.
+ Adds unit tests for the new functionality.

#rnx
#Synced-CL 52552752

diff --git a/Engine/Source/Programs/CSVTools/PerfReportTool.Tests/Integration/CliArgumentTests.cs b/Engine/Source/Programs/CSVTools/PerfReportTool.Tests/Integration/CliArgumentTests.cs
--- a/Engine/Source/Programs/CSVTools/PerfReportTool.Tests/Integration/CliArgumentTests.cs
+++ b/Engine/Source/Programs/CSVTools/PerfReportTool.Tests/Integration/CliArgumentTests.cs
@@ -6,6 +6,7 @@ using System.Diagnostics;
 using System.IO;
 using System.Linq;
 using System.Text.Json;
+using System.Text.Json.Nodes;
 using Xunit;
 
 namespace PerfReportTool.Tests.Integration
@@ -870,5 +871,464 @@ namespace PerfReportTool.Tests.Integration
 		}
 
 		#endregion
+
+		#region JSON output format with Bulk Mode
+
+		[Fact]
+		public void SummaryTableOutputFormatsJson_WithBulkMode_ProducesJsonAndHtml()
+		{
+			_fixture.RequireTestData();
+			string outputDir = Path.Combine(_outputDir, "JsonBulk");
+
+			var (exitCode, output) = RunTool(
+				$"{BulkArgs(outputDir)} " +
+				$"-summaryTableOutputFormats html,json");
+
+			Assert.DoesNotContain("[ERROR]", output);
+			Assert.True(Directory.GetFiles(outputDir, "*.html").Length > 0,
+				"Bulk mode should produce HTML output alongside JSON");
+			Assert.True(Directory.GetFiles(outputDir, "*.json").Length > 0,
+				"JSON output should be written when summaryTableOutputFormats includes json");
+		}
+
+		[Fact]
+		public void SummaryTableOutputFormatsJson_Collated_WithBulkMode_ProducesJsonAndHtml()
+		{
+			_fixture.RequireTestData();
+			string outputDir = Path.Combine(_outputDir, "CollatedJsonBulk");
+
+			var (exitCode, output) = RunTool(
+				$"{BulkArgs(outputDir)} " +
+				$"-summaryTableOutputFormats html,json -collateTable");
+
+			Assert.DoesNotContain("[ERROR]", output);
+			Assert.True(Directory.GetFiles(outputDir, "*.html").Length > 0,
+				"Bulk mode should produce HTML output alongside collated JSON");
+			string[] jsonFiles = Directory.GetFiles(outputDir, "*.json");
+			Assert.True(jsonFiles.Length > 0,
+				"Collated JSON should be written when summaryTableOutputFormats includes json with -collateTable");
+		}
+
+		#endregion
+
+		#region summaryTableOutputFormats combinations
+
+		/// <summary>
+		/// Helper: run bulk mode (without -emailtable) with the given -summaryTableOutputFormats value
+		/// and return the set of output file extensions produced (lowercase, e.g. ".html").
+		/// </summary>
+		private HashSet<string> RunOutputFormats(string tag, string formats, string extraArgs = "")
+		{
+			string outputDir = Path.Combine(_outputDir, tag);
+			var (exitCode, output) = RunTool(
+				$"-csvdir \"{_fixture.GetPerfCsvDir()}\" -recurse -o \"{outputDir}\" " +
+				$"-nowatermarks -nodetailedreports -searchpattern Perf_*.csv.bin " +
+				$"-summaryTableOutputFormats {formats} {extraArgs}");
+
+			Assert.DoesNotContain("[ERROR]", output);
+
+			string[] files = Directory.GetFiles(outputDir);
+			var extensions = new HashSet<string>(
+				files.Select(f => Path.GetExtension(f).ToLowerInvariant()),
+				StringComparer.OrdinalIgnoreCase);
+			return extensions;
+		}
+
+		[Fact]
+		public void OutputFormats_DefaultIsHtmlOnly()
+		{
+			_fixture.RequireTestData();
+			string outputDir = Path.Combine(_outputDir, "FmtDefault");
+			var (_, output) = RunTool(
+				$"-csvdir \"{_fixture.GetPerfCsvDir()}\" -recurse -o \"{outputDir}\" " +
+				$"-nowatermarks -nodetailedreports -searchpattern Perf_*.csv.bin");
+			Assert.DoesNotContain("[ERROR]", output);
+
+			string[] files = Directory.GetFiles(outputDir);
+			var extensions = files.Select(f => Path.GetExtension(f).ToLowerInvariant()).ToHashSet(StringComparer.OrdinalIgnoreCase);
+			Assert.Contains(".html", extensions);
+			Assert.DoesNotContain(".csv", extensions);
+			Assert.DoesNotContain(".json", extensions);
+		}
+
+		[Fact]
+		public void OutputFormats_HtmlOnly_ProducesOnlyHtml()
+		{
+			_fixture.RequireTestData();
+			var ext = RunOutputFormats("FmtHtml", "html");
+			Assert.Contains(".html", ext);
+			Assert.DoesNotContain(".csv", ext);
+			Assert.DoesNotContain(".json", ext);
+		}
+
+		[Fact]
+		public void OutputFormats_CsvOnly_ProducesOnlyCsv()
+		{
+			_fixture.RequireTestData();
+			var ext = RunOutputFormats("FmtCsv", "csv");
+			Assert.Contains(".csv", ext);
+			Assert.DoesNotContain(".html", ext);
+			Assert.DoesNotContain(".json", ext);
+		}
+
+		[Fact]
+		public void OutputFormats_JsonOnly_ProducesOnlyJson()
+		{
+			_fixture.RequireTestData();
+			var ext = RunOutputFormats("FmtJson", "json");
+			Assert.Contains(".json", ext);
+			Assert.DoesNotContain(".html", ext);
+			Assert.DoesNotContain(".csv", ext);
+		}
+
+		[Fact]
+		public void OutputFormats_HtmlAndCsv_ProducesBoth()
+		{
+			_fixture.RequireTestData();
+			var ext = RunOutputFormats("FmtHtmlCsv", "html,csv");
+			Assert.Contains(".html", ext);
+			Assert.Contains(".csv", ext);
+			Assert.DoesNotContain(".json", ext);
+		}
+
+		[Fact]
+		public void OutputFormats_HtmlAndJson_ProducesBoth()
+		{
+			_fixture.RequireTestData();
+			var ext = RunOutputFormats("FmtHtmlJson", "html,json");
+			Assert.Contains(".html", ext);
+			Assert.Contains(".json", ext);
+			Assert.DoesNotContain(".csv", ext);
+		}
+
+		[Fact]
+		public void OutputFormats_CsvAndJson_ProducesBoth()
+		{
+			_fixture.RequireTestData();
+			var ext = RunOutputFormats("FmtCsvJson", "csv,json");
+			Assert.Contains(".csv", ext);
+			Assert.Contains(".json", ext);
+			Assert.DoesNotContain(".html", ext);
+		}
+
+		[Fact]
+		public void OutputFormats_AllThree_ProducesHtmlCsvAndJson()
+		{
+			_fixture.RequireTestData();
+			var ext = RunOutputFormats("FmtAll", "html,csv,json");
+			Assert.Contains(".html", ext);
+			Assert.Contains(".csv", ext);
+			Assert.Contains(".json", ext);
+		}
+
+		[Fact]
+		public void OutputFormats_Collated_ProducesCorrectFiles()
+		{
+			_fixture.RequireTestData();
+			var ext = RunOutputFormats("FmtCollated", "html,csv,json", "-collateTable");
+			Assert.Contains(".html", ext);
+			Assert.Contains(".csv", ext);
+			Assert.Contains(".json", ext);
+		}
+
+		[Fact]
+		public void OutputFormats_DeprecatedCsvTable_ProducesCsvOnly()
+		{
+			_fixture.RequireTestData();
+			string outputDir = Path.Combine(_outputDir, "FmtCsvTableCompat");
+			var (exitCode, output) = RunTool(
+				$"-csvdir \"{_fixture.GetPerfCsvDir()}\" -recurse -o \"{outputDir}\" " +
+				$"-nowatermarks -nodetailedreports -searchpattern Perf_*.csv.bin " +
+				$"-csvTable");
+
+			Assert.DoesNotContain("[ERROR]", output);
+
+			string[] files = Directory.GetFiles(outputDir);
+			var extensions = new HashSet<string>(
+				files.Select(f => Path.GetExtension(f).ToLowerInvariant()),
+				StringComparer.OrdinalIgnoreCase);
+
+			Assert.Contains(".csv", extensions);
+			Assert.DoesNotContain(".html", extensions);
+			Assert.DoesNotContain(".json", extensions);
+		}
+
+		#endregion
+
+		#region JSON output format — compatible arguments
+
+		/// <summary>
+		/// Helper: run with -summaryTableOutputFormats json via PRC cache.
+		/// For collated output, pass bCollated=true which adds -collateTableOnly.
+		/// Returns parsed JSON root and raw output.
+		/// </summary>
+		private (Dictionary<string, JsonElement> root, string output) RunJsonFromCache(
+			string cacheName, bool bCollated, string extraArgs = "")
+		{
+			string cacheDir = GeneratePrcCache(cacheName);
+			string jsonOutputDir = Path.Combine(_outputDir, cacheName + "_out");
+			Directory.CreateDirectory(jsonOutputDir);
+
+			string collateArg = bCollated ? "-collateTableOnly" : "";
+			var (exitCode, output) = RunTool(
+				$"-summarytablecachein \"{cacheDir}\" " +
+				$"-summaryTableOutputFormats json " +
+				$"{collateArg} " +
+				$"-nodetailedreports -nowatermarks " +
+				$"-precacheThreads 1 " +
+				$"-o \"{jsonOutputDir}\" " +
+				extraArgs);
+
+			Assert.DoesNotContain("[ERROR]", output);
+
+			string jsonPath = Path.Combine(jsonOutputDir, "SummaryTable.json");
+			Assert.True(File.Exists(jsonPath), $"JSON not generated for {cacheName} at {jsonPath}");
+
+			string content = File.ReadAllText(jsonPath);
+			var root = JsonSerializer.Deserialize<Dictionary<string, JsonElement>>(content);
+			Assert.NotNull(root);
+			return (root, output);
+		}
+
+		private (Dictionary<string, JsonElement> root, string output) RunFilteredJson(
+			string cacheName, string extraArgs = "")
+			=> RunJsonFromCache(cacheName, bCollated: false, extraArgs);
+
+		private (Dictionary<string, JsonElement> root, string output) RunCollatedJson(
+			string cacheName, string extraArgs = "")
+			=> RunJsonFromCache(cacheName, bCollated: true, extraArgs);
+
+		/// <summary>Collect all column keys across all rows in a flat JSON dict.</summary>
+		private static HashSet<string> CollectAllKeys(Dictionary<string, JsonElement> root)
+		{
+			var keys = new HashSet<string>();
+			foreach (var entry in root)
+			{
+				var rowDict = JsonSerializer.Deserialize<Dictionary<string, JsonElement>>(entry.Value.GetRawText());
+				foreach (string key in rowDict.Keys)
+				{
+					keys.Add(key);
+				}
+			}
+			return keys;
+		}
+
+		// -customTable: override the stat filter list entirely
+		[Fact]
+		public void JsonFiltered_CustomTable_OnlyIncludesSpecifiedStats()
+		{
+			_fixture.RequireTestData();
+
+			var defaultKeys = CollectAllKeys(RunFilteredJson("JFCustomBase").root);
+			var customKeys = CollectAllKeys(RunFilteredJson("JFCustom", "-customTable MVP,HitchTimePercent").root);
+
+			Assert.True(customKeys.Count < defaultKeys.Count,
+				$"Custom filter should produce fewer columns ({customKeys.Count}) than default ({defaultKeys.Count})");
+			Assert.True(customKeys.Contains("MVP"), "Custom filter should include MVP");
+			Assert.True(customKeys.Contains("HitchTimePercent"), "Custom filter should include HitchTimePercent");
+		}
+
+		// -customTableSort: override the row sort columns
+		[Fact]
+		public void JsonCollated_CustomTableSort_ChangesCollationKeys()
+		{
+			_fixture.RequireTestData();
+
+			var defaultRoot = RunCollatedJson("JCSort_Default").root;
+			var customRoot = RunCollatedJson("JCSort_Custom", "-customTableSort Platform").root;
+
+			Assert.True(customRoot.Count > 0);
+			// With a broader sort key the default will produce more groups;
+			// with just Platform there should be fewer or equal groups
+			Assert.True(customRoot.Count <= defaultRoot.Count,
+				$"Sorting by fewer columns should produce fewer groups ({customRoot.Count}) " +
+				$"than the default ({defaultRoot.Count})");
+		}
+
+		// -summaryTable <name>: select a named summary table definition from the XML
+		[Fact]
+		public void JsonFiltered_SummaryTableName_UsesNamedDefinition()
+		{
+			_fixture.RequireTestData();
+
+			// Use the "memory" summary table (filters to MemoryFreeMB*, PhysicalUsedMB*, Total Time*)
+			var keys = CollectAllKeys(RunFilteredJson("JFMemory", "-summaryTable memory").root);
+
+			// "memory" table should NOT include stats from "default" like MVP, Hitches, etc.
+			Assert.False(keys.Any(k => k.StartsWith("MVP")),
+				"Memory summary table should not include MVP stats");
+			Assert.False(keys.Any(k => k.StartsWith("Hitches")),
+				"Memory summary table should not include Hitches stats");
+		}
+
+		// -summaryTableXmlAppend: append additional stats to the filter list
+		[Fact]
+		public void JsonFiltered_SummaryTableXmlAppend_AddsExtraStats()
+		{
+			_fixture.RequireTestData();
+
+			var baseKeys = CollectAllKeys(RunFilteredJson("JFAppend_Base", "-summaryTable condensed").root);
+			var appendKeys = CollectAllKeys(RunFilteredJson("JFAppend_Extra",
+				"-summaryTable condensed -summaryTableXmlAppend framecount*").root);
+
+			Assert.True(appendKeys.Count >= baseKeys.Count,
+				$"Appending stats should add columns: base={baseKeys.Count}, appended={appendKeys.Count}");
+			Assert.True(appendKeys.Any(k => k.ToLower().Contains("framecount")),
+				"Appended stat 'framecount*' should appear in the output");
+		}
+
+		// -weightByColumn: weight collated averages by a specific column
+		[Fact]
+		public void JsonCollated_WeightByColumn_ProducesValidOutput()
+		{
+			_fixture.RequireTestData();
+
+			var defaultRoot = RunCollatedJson("JCWeight_Default").root;
+			var weightedRoot = RunCollatedJson("JCWeight_By", "-weightByColumn framecount").root;
+
+			Assert.True(defaultRoot.Count > 0);
+			Assert.True(weightedRoot.Count > 0);
+			// Both should produce valid output with the same row keys
+			Assert.Equal(defaultRoot.Count, weightedRoot.Count);
+		}
+
+		// -noWeightedAvg: disable weighted averages
+		[Fact]
+		public void JsonCollated_NoWeightedAvg_ProducesValidOutput()
+		{
+			_fixture.RequireTestData();
+			Assert.True(RunCollatedJson("JCNoWeight", "-noWeightedAvg").root.Count > 0);
+		}
+
+		// -reverseTable: reverse the row order
+		[Fact]
+		public void JsonFiltered_ReverseTable_ReversesRowOrder()
+		{
+			_fixture.RequireTestData();
+
+			// Generate both variants (validates no errors via RunFilteredJson assertions)
+			RunFilteredJson("JFReverse_Normal");
+			RunFilteredJson("JFReverse_Rev", "-reverseTable");
+
+			// Re-read with JsonNode.Parse which preserves document key order,
+			// unlike Dictionary<string,T> whose enumeration order is unspecified.
+			string normalPath = Path.Combine(_outputDir, "JFReverse_Normal_out", "SummaryTable.json");
+			string reversedPath = Path.Combine(_outputDir, "JFReverse_Rev_out", "SummaryTable.json");
+			var normalKeys = JsonNode.Parse(File.ReadAllText(normalPath))!.AsObject().Select(kv => kv.Key).ToList();
+			var reversedKeys = JsonNode.Parse(File.ReadAllText(reversedPath))!.AsObject().Select(kv => kv.Key).ToList();
+
+			Assert.Equal(normalKeys.Count, reversedKeys.Count);
+			if (normalKeys.Count >= 2)
+			{
+				// First row in normal should be last in reversed (or vice versa)
+				Assert.Equal(normalKeys.First(), reversedKeys.Last());
+				Assert.Equal(normalKeys.Last(), reversedKeys.First());
+			}
+		}
+
+		// -summaryTableStatThreshold: filter out low-value stat columns
+		[Fact]
+		public void JsonFiltered_StatThreshold_FiltersLowValueColumns()
+		{
+			_fixture.RequireTestData();
+
+			var baseKeys = CollectAllKeys(RunFilteredJson("JFThresh_Base").root);
+			var filteredKeys = CollectAllKeys(RunFilteredJson("JFThresh_High",
+				"-summaryTableStatThreshold 999999").root);
+
+			Assert.True(filteredKeys.Count < baseKeys.Count,
+				$"High stat threshold should remove columns: base={baseKeys.Count}, filtered={filteredKeys.Count}");
+		}
+
+		// -hideMetadataColumns: remove metadata columns from the output
+		[Fact]
+		public void JsonFiltered_HideMetadataColumns_RemovesMetadata()
+		{
+			_fixture.RequireTestData();
+
+			var withKeys = CollectAllKeys(RunFilteredJson("JFHideMeta_With").root);
+			var withoutKeys = CollectAllKeys(RunFilteredJson("JFHideMeta_Without",
+				"-hideMetadataColumns").root);
+
+			Assert.True(withoutKeys.Count < withKeys.Count,
+				$"Hiding metadata should remove columns: with={withKeys.Count}, without={withoutKeys.Count}");
+		}
+
+		// -collatedStringVisibility: control string column visibility in collated output
+		[Theory]
+		[InlineData("show")]
+		[InlineData("hide")]
+		[InlineData("auto")]
+		public void JsonCollated_CollatedStringVisibility_ProducesValidOutput(string visibility)
+		{
+			_fixture.RequireTestData();
+			Assert.True(RunCollatedJson($"JCStrVis_{visibility}",
+				$"-collatedStringVisibility {visibility}").root.Count > 0);
+		}
+
+		// -collatedDateVisibility: control date column visibility
+		[Theory]
+		[InlineData("hide")]
+		[InlineData("newest")]
+		[InlineData("oldest")]
+		public void JsonCollated_CollatedDateVisibility_ProducesValidOutput(string visibility)
+		{
+			_fixture.RequireTestData();
+			Assert.True(RunCollatedJson($"JCDateVis_{visibility}",
+				$"-collatedDateVisibility {visibility}").root.Count > 0);
+		}
+
+		// -columnSortMode: various column sorting modes
+		[Theory]
+		[InlineData("sortByMax")]
+		[InlineData("sortByAvg")]
+		public void JsonFiltered_ColumnSortMode_ProducesValidOutput(string mode)
+		{
+			_fixture.RequireTestData();
+			Assert.True(RunFilteredJson($"JFColSort_{mode}",
+				$"-columnSortMode {mode}").root.Count > 0);
+		}
+
+		// -sortTrailingDigitsAsNumeric
+		[Fact]
+		public void JsonFiltered_SortTrailingDigitsAsNumeric_ProducesValidOutput()
+		{
+			_fixture.RequireTestData();
+			Assert.True(RunFilteredJson("JFNumDigitSort",
+				"-sortTrailingDigitsAsNumeric").root.Count > 0);
+		}
+
+		// -customTable + -customTableSort + collated JSON + -noSummaryMinMax
+		// Note: the sort column (platform) must be in -customTable for collation to work,
+		// since CollateSortedTable requires the collate-by columns to be present in the filtered table.
+		[Fact]
+		public void JsonCollated_CustomTable_WithNoMinMax_ProducesFilteredCollatedOutput()
+		{
+			_fixture.RequireTestData();
+			var keys = CollectAllKeys(RunCollatedJson("JCCustomCombo",
+				"-customTable MVP,HitchTimePercent,Hitches/Min,platform -customTableSort platform -nosummaryminmax").root);
+
+			Assert.False(keys.Any(k => k.StartsWith("Min ")),
+				"noSummaryMinMax should suppress Min columns");
+			Assert.False(keys.Any(k => k.StartsWith("Max ")),
+				"noSummaryMinMax should suppress Max columns");
+			Assert.True(keys.Any(k => k.Contains("MVP")),
+				"Custom filter should include MVP stat");
+		}
+
+		// Full combo: -summaryTable + -summaryTableXmlAppend + -weightByColumn + -collatedStringVisibility
+		[Fact]
+		public void JsonCollated_FullCombo_AllArgsWorkTogether()
+		{
+			_fixture.RequireTestData();
+			var keys = CollectAllKeys(RunCollatedJson("JCFullCombo",
+				"-summaryTable condensed -summaryTableXmlAppend framecount* " +
+				"-weightByColumn framecount -collatedStringVisibility hide").root);
+
+			Assert.True(keys.Contains("Count"), "Collated output should have Count column");
+			Assert.True(keys.Any(k => k.StartsWith("Avg ")), "Collated output should have Avg columns");
+		}
+
+		#endregion
 	}
 }
diff --git a/Engine/Source/Programs/CSVTools/PerfReportTool.Tests/Integration/EndToEndTests.cs b/Engine/Source/Programs/CSVTools/PerfReportTool.Tests/Integration/EndToEndTests.cs
--- a/Engine/Source/Programs/CSVTools/PerfReportTool.Tests/Integration/EndToEndTests.cs
+++ b/Engine/Source/Programs/CSVTools/PerfReportTool.Tests/Integration/EndToEndTests.cs
@@ -447,5 +447,169 @@ namespace PerfReportTool.Tests.Integration
 		}
 
 		#endregion
+
+		#region Filtered/Collated JSON tests (via -summaryTableOutputFormats=json)
+
+		/// <summary>
+		/// Helper: run with -summaryTableOutputFormats=json via PRC cache.
+		/// For collated output, pass bCollated=true which adds -collateTableOnly.
+		/// Returns parsed JSON root and raw output.
+		/// </summary>
+		private (Dictionary<string, JsonElement> root, string output) RunOutputFormatJsonFromCache(
+			string cacheName, bool bCollated, string extraArgs = "")
+		{
+			string cacheDir = GeneratePerfPrcCache(cacheName);
+			string jsonOutputDir = Path.Combine(_outputDir, cacheName + "_out");
+			Directory.CreateDirectory(jsonOutputDir);
+
+			string collateArg = bCollated ? "-collateTableOnly" : "";
+			var (exitCode, output) = RunPerfReportTool(
+				$"-summarytablecachein \"{cacheDir}\" " +
+				$"-summaryTableOutputFormats json " +
+				$"{collateArg} " +
+				$"-nodetailedreports -nowatermarks " +
+				$"-precacheThreads 1 " +
+				$"-o \"{jsonOutputDir}\" " +
+				extraArgs,
+				timeoutMs: 120000);
+
+			Assert.DoesNotContain("[ERROR]", output);
+
+			string jsonPath = Path.Combine(jsonOutputDir, "SummaryTable.json");
+			Assert.True(File.Exists(jsonPath), $"JSON not generated for {cacheName} at {jsonPath}");
+
+			string content = File.ReadAllText(jsonPath);
+			var root = JsonSerializer.Deserialize<Dictionary<string, JsonElement>>(content);
+			Assert.NotNull(root);
+			return (root, output);
+		}
+
+		[Fact]
+		public void SummaryTableOutputFormatsJson_ProducesFlatJson()
+		{
+			_fixture.RequireTestData();
+
+			var (root, _) = RunOutputFormatJsonFromCache("JsonFiltered", bCollated: false);
+
+			Assert.True(root.Count > 0, "JSON should contain at least one entry");
+
+			// Flat format: no type-sectioned nesting like CsvStatAverage/CsvMetadata
+			foreach (var entry in root)
+			{
+				var rowDict = JsonSerializer.Deserialize<Dictionary<string, JsonElement>>(entry.Value.GetRawText());
+				Assert.False(rowDict.ContainsKey("CsvStatAverage"),
+					"Output format JSON should use flat format, not type-sectioned format");
+				Assert.False(rowDict.ContainsKey("CsvMetadata"),
+					"Output format JSON should use flat format, not type-sectioned format");
+			}
+		}
+
+		[Fact]
+		public void SummaryTableOutputFormatsJson_OnlyIncludesFilteredStats()
+		{
+			_fixture.RequireTestData();
+
+			string cacheDir = GeneratePerfPrcCache("JsonFilteredStats");
+
+			// Unfiltered (old -summaryTableToJson path which dumps all row data)
+			string unfilteredPath = Path.Combine(_outputDir, "Unfiltered.json");
+			RunPerfReportTool(
+				$"-summarytablecachein \"{cacheDir}\" " +
+				$"-summarytabletojsonfastmode " +
+				$"-summaryTableToJson \"{unfilteredPath}\" " +
+				$"-precacheThreads 1",
+				timeoutMs: 120000);
+
+			// Filtered (via -summaryTableOutputFormats=json)
+			string filteredOutputDir = Path.Combine(_outputDir, "FilteredStats_out");
+			Directory.CreateDirectory(filteredOutputDir);
+			RunPerfReportTool(
+				$"-summarytablecachein \"{cacheDir}\" " +
+				$"-summaryTableOutputFormats json " +
+				$"-nodetailedreports -nowatermarks " +
+				$"-precacheThreads 1 " +
+				$"-o \"{filteredOutputDir}\"",
+				timeoutMs: 120000);
+
+			Assert.True(File.Exists(unfilteredPath), "Unfiltered JSON not generated");
+			string filteredPath = Path.Combine(filteredOutputDir, "SummaryTable.json");
+			Assert.True(File.Exists(filteredPath), "Filtered JSON not generated");
+
+			string unfilteredContent = File.ReadAllText(unfilteredPath);
+			string filteredContent = File.ReadAllText(filteredPath);
+
+			// The filtered file should be smaller since it only includes stats from the summary table filter
+			Assert.True(filteredContent.Length < unfilteredContent.Length,
+				"Filtered JSON should contain less data than unfiltered JSON");
+		}
+
+		[Fact]
+		public void SummaryTableOutputFormatsJson_Collated_ProducesCollatedFlatJson()
+		{
+			_fixture.RequireTestData();
+
+			var (root, _) = RunOutputFormatJsonFromCache("JsonCollated", bCollated: true);
+
+			Assert.True(root.Count > 0, "JSON should contain at least one collated entry");
+
+			// Collated rows should have aggregate column names and a Count field
+			bool foundAvgColumn = false;
+			bool foundCountColumn = false;
+			foreach (var entry in root)
+			{
+				var rowDict = JsonSerializer.Deserialize<Dictionary<string, JsonElement>>(entry.Value.GetRawText());
+				foreach (string key in rowDict.Keys)
+				{
+					if (key.StartsWith("Avg ")) foundAvgColumn = true;
+					if (key == "Count") foundCountColumn = true;
+				}
+				// Flat format, not type-sectioned
+				Assert.False(rowDict.ContainsKey("CsvStatAverage"),
+					"Collated JSON should use flat format, not type-sectioned format");
+			}
+			Assert.True(foundAvgColumn, "Collated JSON should contain Avg-prefixed columns");
+			Assert.True(foundCountColumn, "Collated JSON should contain a Count column");
+		}
+
+		[Fact]
+		public void SummaryTableOutputFormatsJson_Collated_FewerRowsThanNonCollated()
+		{
+			_fixture.RequireTestData();
+
+			string cacheDir = GeneratePerfPrcCache("JsonCollatedRows");
+
+			// Non-collated
+			string filteredOutputDir = Path.Combine(_outputDir, "FilteredRows_out");
+			Directory.CreateDirectory(filteredOutputDir);
+			RunPerfReportTool(
+				$"-summarytablecachein \"{cacheDir}\" " +
+				$"-summaryTableOutputFormats json " +
+				$"-nodetailedreports -nowatermarks " +
+				$"-precacheThreads 1 -o \"{filteredOutputDir}\"",
+				timeoutMs: 120000);
+
+			// Collated
+			string collatedOutputDir = Path.Combine(_outputDir, "CollatedRows_out");
+			Directory.CreateDirectory(collatedOutputDir);
+			RunPerfReportTool(
+				$"-summarytablecachein \"{cacheDir}\" " +
+				$"-summaryTableOutputFormats json -collateTableOnly " +
+				$"-nodetailedreports -nowatermarks " +
+				$"-precacheThreads 1 -o \"{collatedOutputDir}\"",
+				timeoutMs: 120000);
+
+			string filteredPath = Path.Combine(filteredOutputDir, "SummaryTable.json");
+			string collatedPath = Path.Combine(collatedOutputDir, "SummaryTable.json");
+			Assert.True(File.Exists(filteredPath));
+			Assert.True(File.Exists(collatedPath));
+
+			var filteredRoot = JsonSerializer.Deserialize<Dictionary<string, JsonElement>>(File.ReadAllText(filteredPath));
+			var collatedRoot = JsonSerializer.Deserialize<Dictionary<string, JsonElement>>(File.ReadAllText(collatedPath));
+
+			Assert.True(collatedRoot.Count <= filteredRoot.Count,
+				$"Collated rows ({collatedRoot.Count}) should be <= filtered rows ({filteredRoot.Count})");
+		}
+
+		#endregion
 	}
 }
diff --git a/Engine/Source/Programs/CSVTools/PerfReportTool.Tests/PerfReportToolTests/CommandLineTests.cs b/Engine/Source/Programs/CSVTools/PerfReportTool.Tests/PerfReportToolTests/CommandLineTests.cs
--- a/Engine/Source/Programs/CSVTools/PerfReportTool.Tests/PerfReportToolTests/CommandLineTests.cs
+++ b/Engine/Source/Programs/CSVTools/PerfReportTool.Tests/PerfReportToolTests/CommandLineTests.cs
@@ -260,6 +260,87 @@ namespace PerfReportTool.Tests.PerfReportToolTests
 		#endregion
 	}
 
+	[Flags]
+	public enum TestFlags
+	{
+		None  = 0,
+		Alpha = 1,
+		Beta  = 2,
+		Gamma = 4,
+		Delta = 8,
+	}
+
+	// Minimal subclass to expose protected CommandLineTool methods for testing
+	public class TestableCommandLineTool : CommandLineTool
+	{
+		public TestableCommandLineTool(string[] args)
+		{
+			ReadCommandLine(args);
+		}
+
+		public T TestGetEnumArg<T>(string key, T defaultValue) where T : System.Enum
+			=> GetEnumArg(key, defaultValue);
+	}
+
+	public class GetEnumArgTests
+	{
+		[Fact]
+		public void SingleValue_ReturnsThatFlag()
+		{
+			var tool = new TestableCommandLineTool(new[] { "-flags", "Alpha" });
+			Assert.Equal(TestFlags.Alpha, tool.TestGetEnumArg("flags", TestFlags.None));
+		}
+
+		[Fact]
+		public void MultipleValues_ReturnsOrCombined()
+		{
+			var tool = new TestableCommandLineTool(new[] { "-flags", "Alpha,Gamma" });
+			Assert.Equal(TestFlags.Alpha | TestFlags.Gamma, tool.TestGetEnumArg("flags", TestFlags.None));
+		}
+
+		[Fact]
+		public void AllValues_ReturnsFullCombination()
+		{
+			var tool = new TestableCommandLineTool(new[] { "-flags", "Alpha,Beta,Gamma,Delta" });
+			Assert.Equal(TestFlags.Alpha | TestFlags.Beta | TestFlags.Gamma | TestFlags.Delta, tool.TestGetEnumArg("flags", TestFlags.None));
+		}
+
+		[Fact]
+		public void CaseInsensitive_ParsesCorrectly()
+		{
+			var tool = new TestableCommandLineTool(new[] { "-flags", "alpha,BETA,gAmMa" });
+			Assert.Equal(TestFlags.Alpha | TestFlags.Beta | TestFlags.Gamma, tool.TestGetEnumArg("flags", TestFlags.None));
+		}
+
+		[Fact]
+		public void WhitespaceAroundValues_IsTrimmed()
+		{
+			var tool = new TestableCommandLineTool(new[] { "-flags", "Alpha , Beta , Gamma" });
+			Assert.Equal(TestFlags.Alpha | TestFlags.Beta | TestFlags.Gamma, tool.TestGetEnumArg("flags", TestFlags.None));
+		}
+
+		[Fact]
+		public void MissingArg_ReturnsDefault()
+		{
+			var tool = new TestableCommandLineTool(new[] { "-other", "value" });
+			Assert.Equal(TestFlags.Beta, tool.TestGetEnumArg("flags", TestFlags.Beta));
+		}
+
+		[Fact]
+		public void UnrecognizedValue_ReturnsDefault()
+		{
+			var tool = new TestableCommandLineTool(new[] { "-flags", "Bogus" });
+			Assert.Equal(TestFlags.Beta, tool.TestGetEnumArg("flags", TestFlags.Beta));
+		}
+
+		[Fact]
+		public void DuplicateValues_SameAsOnce()
+		{
+			var tool = new TestableCommandLineTool(new[] { "-flags", "Alpha,Alpha,Alpha" });
+			Assert.Equal(TestFlags.Alpha, tool.TestGetEnumArg("flags", TestFlags.None));
+		}
+	}
+
 	public class PerfLogTests
 	{
 		[Fact]
diff --git a/Engine/Source/Programs/CSVTools/PerfReportTool.Tests/PerfReportToolTests/SummaryTableTests.cs b/Engine/Source/Programs/CSVTools/PerfReportTool.Tests/PerfReportToolTests/SummaryTableTests.cs
--- a/Engine/Source/Programs/CSVTools/PerfReportTool.Tests/PerfReportToolTests/SummaryTableTests.cs
+++ b/Engine/Source/Programs/CSVTools/PerfReportTool.Tests/PerfReportToolTests/SummaryTableTests.cs
@@ -601,6 +601,361 @@ namespace PerfReportTool.Tests.PerfReportToolTests
 		}
 
 		#endregion
+
+		#region ToFlatJsonDict
+
+		/// <summary>Helper: add a row with optional csvid, platform metadata, and a frametime stat.</summary>
+		private static void AddRow(SummaryTable table, string csvId, string platform, double frametime)
+		{
+			var row = new SummaryTableRowData();
+			if (csvId != null) row.Add(SummaryTableElement.Type.CsvMetadata, "csvid", csvId);
+			row.Add(SummaryTableElement.Type.CsvMetadata, "platform", platform);
+			row.Add(SummaryTableElement.Type.CsvStatAverage, "frametime", frametime);
+			table.AddRowData(row, true, false, false, false);
+		}
+
+		/// <summary>Helper: sort by platform then collate.</summary>
+		private static SummaryTable SortAndCollateByPlatform(SummaryTable table, bool addMinMax = false)
+		{
+			var sorted = table.SortRows(new List<string> { "platform" }, false);
+			return sorted.CollateSortedTable(
+				new List<string> { "platform" }, addMinMax,
+				StringCollationVisibility.Auto, DateCollationVisibility.Newest);
+		}
+
+		/// <summary>Helper: build a 3-row table (2x windows + 1x linux).</summary>
+		private static SummaryTable BuildThreeRowTable()
+		{
+			var table = new SummaryTable();
+			AddRow(table, "run_001", "windows", 16.0);
+			AddRow(table, "run_002", "windows", 20.0);
+			AddRow(table, "run_003", "linux", 33.0);
+			return table;
+		}
+
+		/// <summary>Helper: build a 2-row table (2x windows, same platform for collation).</summary>
+		private static SummaryTable BuildTwoWindowsRowTable()
+		{
+			var table = new SummaryTable();
+			AddRow(table, "run_001", "windows", 16.0);
+			AddRow(table, "run_002", "windows", 20.0);
+			return table;
+		}
+
+		[Fact]
+		public void ToFlatJsonDict_KeysByCsvId()
+		{
+			var dict = BuildTwoWindowsRowTable().ToFlatJsonDict();
+
+			Assert.Equal(2, dict.Count);
+			Assert.True(dict.ContainsKey("run_001"));
+			Assert.True(dict.ContainsKey("run_002"));
+		}
+
+		[Fact]
+		public void ToFlatJsonDict_FallsBackToRowIndex_WhenNoCsvId()
+		{
+			var table = new SummaryTable();
+			AddRow(table, null, "windows", 16.0);
+			AddRow(table, null, "linux", 20.0);
+
+			var dict = table.ToFlatJsonDict();
+
+			Assert.Equal(2, dict.Count);
+			Assert.True(dict.ContainsKey("0"));
+			Assert.True(dict.ContainsKey("1"));
+		}
+
+		[Fact]
+		public void ToFlatJsonDict_IncludesNumericAndStringValues()
+		{
+			var table = new SummaryTable();
+			var row = new SummaryTableRowData();
+			row.Add(SummaryTableElement.Type.CsvMetadata, "csvid", "run_001");
+			row.Add(SummaryTableElement.Type.CsvMetadata, "platform", "windows");
+			row.Add(SummaryTableElement.Type.CsvStatAverage, "frametime", 16.6);
+			row.Add(SummaryTableElement.Type.CsvStatAverage, "drawcalls", 500.0);
+			table.AddRowData(row, true, false, false, false);
+
+			var dict = table.ToFlatJsonDict();
+			Dictionary<string, dynamic> rowDict = dict["run_001"];
+
+			Assert.Equal("windows", (string)rowDict["platform"]);
+			Assert.Equal(16.6m, (decimal)rowDict["frametime"]);
+			Assert.Equal(500.0m, (decimal)rowDict["drawcalls"]);
+		}
+
+		[Fact]
+		public void ToFlatJsonDict_OmitsMaxValueSentinels()
+		{
+			var table = new SummaryTable();
+			// row1 has drawcalls, row2 doesn't — column will contain double.MaxValue sentinel
+			var row1 = new SummaryTableRowData();
+			row1.Add(SummaryTableElement.Type.CsvMetadata, "csvid", "run_001");
+			row1.Add(SummaryTableElement.Type.CsvStatAverage, "frametime", 16.0);
+			row1.Add(SummaryTableElement.Type.CsvStatAverage, "drawcalls", 500.0);
+			table.AddRowData(row1, true, false, false, false);
+
+			AddRow(table, "run_002", "windows", 20.0);
+
+			var dict = table.ToFlatJsonDict();
+			Dictionary<string, dynamic> row2Dict = dict["run_002"];
+
+			Assert.True(row2Dict.ContainsKey("frametime"));
+			Assert.False(row2Dict.ContainsKey("drawcalls"), "MaxValue sentinel should be omitted");
+		}
+
+		[Fact]
+		public void ToFlatJsonDict_OmitsEmptyStrings()
+		{
+			var table = new SummaryTable();
+			var row1 = new SummaryTableRowData();
+			row1.Add(SummaryTableElement.Type.CsvMetadata, "csvid", "run_001");
+			row1.Add(SummaryTableElement.Type.CsvMetadata, "config", "shipping");
+			table.AddRowData(row1, true, false, false, false);
+
+			var row2 = new SummaryTableRowData();
+			row2.Add(SummaryTableElement.Type.CsvMetadata, "csvid", "run_002");
+			table.AddRowData(row2, true, false, false, false);
+
+			var dict = table.ToFlatJsonDict();
+			Assert.False(dict["run_002"].ContainsKey("config"), "Empty string values should be omitted");
+		}
+
+		[Fact]
+		public void ToFlatJsonDict_OmitsInvisibleColumns()
+		{
+			var table = new SummaryTable();
+			var row = new SummaryTableRowData();
+			row.Add(SummaryTableElement.Type.CsvMetadata, "csvid", "run_001");
+			row.Add(SummaryTableElement.Type.CsvStatAverage, "frametime", 16.0);
+			row.Add(SummaryTableElement.Type.CsvStatAverage, "hidden_stat", 99.0);
+			table.AddRowData(row, true, false, false, false);
+
+			table.GetColumnByName("hidden_stat").isVisible = false;
+
+			var dict = table.ToFlatJsonDict();
+			Dictionary<string, dynamic> rowDict = dict["run_001"];
+
+			Assert.True(rowDict.ContainsKey("frametime"));
+			Assert.False(rowDict.ContainsKey("hidden_stat"), "Invisible columns should be omitted");
+		}
+
+		[Fact]
+		public void ToFlatJsonDict_CollatedTable_KeysBySortColumns()
+		{
+			var dict = SortAndCollateByPlatform(BuildThreeRowTable()).ToFlatJsonDict();
+
+			Assert.Equal(2, dict.Count);
+			Assert.True(dict.ContainsKey("linux"), "Collated key should be the sort column value");
+			Assert.True(dict.ContainsKey("windows"), "Collated key should be the sort column value");
+		}
+
+		[Fact]
+		public void ToFlatJsonDict_CollatedTable_MultiSortKey_JoinsWithUnderscore()
+		{
+			var table = new SummaryTable();
+			var row1 = new SummaryTableRowData();
+			row1.Add(SummaryTableElement.Type.CsvMetadata, "csvid", "run_001");
+			row1.Add(SummaryTableElement.Type.CsvMetadata, "platform", "windows");
+			row1.Add(SummaryTableElement.Type.CsvMetadata, "config", "shipping");
+			row1.Add(SummaryTableElement.Type.CsvStatAverage, "frametime", 16.0);
+			table.AddRowData(row1, true, false, false, false);
+
+			var row2 = new SummaryTableRowData();
+			row2.Add(SummaryTableElement.Type.CsvMetadata, "csvid", "run_002");
+			row2.Add(SummaryTableElement.Type.CsvMetadata, "platform", "windows");
+			row2.Add(SummaryTableElement.Type.CsvMetadata, "config", "development");
+			row2.Add(SummaryTableElement.Type.CsvStatAverage, "frametime", 33.0);
+			table.AddRowData(row2, true, false, false, false);
+
+			var sorted = table.SortRows(new List<string> { "platform", "config" }, false);
+			var collated = sorted.CollateSortedTable(
+				new List<string> { "platform", "config" }, false,
+				StringCollationVisibility.Auto, DateCollationVisibility.Newest);
+
+			var dict = collated.ToFlatJsonDict();
+
+			Assert.Equal(2, dict.Count);
+			Assert.True(dict.ContainsKey("windows_development"));
+			Assert.True(dict.ContainsKey("windows_shipping"));
+		}
+
+		[Fact]
+		public void ToFlatJsonDict_CollatedTable_AveragesNumericValues()
+		{
+			var dict = SortAndCollateByPlatform(BuildTwoWindowsRowTable()).ToFlatJsonDict();
+			// The avg of 16.0 and 20.0 is 18.0
+			Assert.Equal(18.0m, (decimal)dict["windows"]["Avg frametime"]);
+		}
+
+		[Fact]
+		public void ToFlatJsonDict_CollatedTable_WithMinMax_IncludesAggregateColumns()
+		{
+			var dict = SortAndCollateByPlatform(BuildTwoWindowsRowTable(), addMinMax: true).ToFlatJsonDict();
+			Dictionary<string, dynamic> windowsRow = dict["windows"];
+
+			Assert.Equal(18.0m, (decimal)windowsRow["Avg frametime"]);
+			Assert.Equal(16.0m, (decimal)windowsRow["Min frametime"]);
+			Assert.Equal(20.0m, (decimal)windowsRow["Max frametime"]);
+		}
+
+		[Fact]
+		public void ToFlatJsonDict_AfterSortAndFilter_OnlyIncludesFilteredColumns()
+		{
+			var table = new SummaryTable();
+			var row = new SummaryTableRowData();
+			row.Add(SummaryTableElement.Type.CsvMetadata, "csvid", "run_001");
+			row.Add(SummaryTableElement.Type.CsvStatAverage, "frametime", 16.6);
+			row.Add(SummaryTableElement.Type.CsvStatAverage, "drawcalls", 500.0);
+			table.AddRowData(row, true, false, false, false);
+
+			var filtered = table.SortAndFilter(
+				new List<string> { "frametime" }, new List<string>(),
+				false, null, null, false, TableColumnSortMode.Default);
+
+			var dict = filtered.ToFlatJsonDict();
+			Dictionary<string, dynamic> rowDict = dict["run_001"];
+
+			Assert.True(rowDict.ContainsKey("frametime"));
+			Assert.False(rowDict.ContainsKey("drawcalls"), "Filtered-out columns should not appear in JSON");
+		}
+
+		[Fact]
+		public void ToFlatJsonDict_IncludesCountColumn_InCollatedTable()
+		{
+			var dict = SortAndCollateByPlatform(BuildThreeRowTable()).ToFlatJsonDict();
+
+			Assert.Equal(2.0m, (decimal)dict["windows"]["Count"]);
+			Assert.Equal(1.0m, (decimal)dict["linux"]["Count"]);
+		}
+
+		[Fact]
+		public void ToFlatJsonDict_EmptyTable_ReturnsEmptyDict()
+		{
+			Assert.Empty(new SummaryTable().ToFlatJsonDict());
+		}
+
+		#endregion
+	}
+
+	public class SummaryTableJsonWriteTests : IDisposable
+	{
+		private readonly string _tempDir;
+
+		public SummaryTableJsonWriteTests()
+		{
+			_tempDir = Path.Combine(Path.GetTempPath(), "prt_json_test_" + Path.GetRandomFileName());
+			Directory.CreateDirectory(_tempDir);
+		}
+
+		public void Dispose()
+		{
+			if (Directory.Exists(_tempDir))
+			{
+				Directory.Delete(_tempDir, true);
+			}
+		}
+
+		private static Dictionary<string, dynamic> MakeSampleDict()
+		{
+			return new Dictionary<string, dynamic>
+			{
+				["run_001"] = new Dictionary<string, dynamic>
+				{
+					["frametime"] = new decimal(16.6),
+					["platform"] = "windows"
+				}
+			};
+		}
+
+		private string WriteAndRead(Dictionary<string, dynamic> dict, string name, bool stream = false, bool indent = true)
+		{
+			string filePath = Path.Combine(_tempDir, name);
+			SummaryTableDataJsonWriteHelper.WriteJsonFile(stream, filePath, dict, indent);
+			Assert.True(File.Exists(filePath));
+			return File.ReadAllText(filePath);
+		}
+
+		[Fact]
+		[Trait("Category", "FileIO")]
+		public void WriteJsonFile_ProducesValidJsonFile()
+		{
+			string content = WriteAndRead(MakeSampleDict(), "test_output.json");
+			Assert.Contains("run_001", content);
+			Assert.Contains("frametime", content);
+			Assert.Contains("16.6", content);
+			Assert.Contains("windows", content);
+		}
+
+		[Fact]
+		[Trait("Category", "FileIO")]
+		public void WriteJsonFile_NoIndent_ProducesCompactJson()
+		{
+			string content = WriteAndRead(MakeSampleDict(), "compact.json", indent: false);
+			Assert.DoesNotContain("\n", content.Trim());
+		}
+
+		[Fact]
+		[Trait("Category", "FileIO")]
+		public void WriteJsonFile_StreamMode_ProducesValidFile()
+		{
+			string content = WriteAndRead(MakeSampleDict(), "streamed.json", stream: true, indent: false);
+			Assert.Contains("run_001", content);
+		}
+
+		[Fact]
+		[Trait("Category", "FileIO")]
+		public void ToFlatJsonDict_RoundTrips_ThroughWriteJsonFile()
+		{
+			var table = new SummaryTable();
+			var row1 = new SummaryTableRowData();
+			row1.Add(SummaryTableElement.Type.CsvMetadata, "csvid", "run_001");
+			row1.Add(SummaryTableElement.Type.CsvMetadata, "platform", "windows");
+			row1.Add(SummaryTableElement.Type.CsvStatAverage, "frametime", 16.6);
+			table.AddRowData(row1, true, false, false, false);
+
+			var row2 = new SummaryTableRowData();
+			row2.Add(SummaryTableElement.Type.CsvMetadata, "csvid", "run_002");
+			row2.Add(SummaryTableElement.Type.CsvMetadata, "platform", "linux");
+			row2.Add(SummaryTableElement.Type.CsvStatAverage, "frametime", 33.3);
+			table.AddRowData(row2, true, false, false, false);
+
+			string content = WriteAndRead(table.ToFlatJsonDict(), "roundtrip.json");
+			Assert.Contains("run_001", content);
+			Assert.Contains("run_002", content);
+			Assert.Contains("windows", content);
+			Assert.Contains("linux", content);
+		}
+
+		[Fact]
+		[Trait("Category", "FileIO")]
+		public void CollatedTable_RoundTrips_ThroughWriteJsonFile()
+		{
+			var table = new SummaryTable();
+			var row1 = new SummaryTableRowData();
+			row1.Add(SummaryTableElement.Type.CsvMetadata, "csvid", "run_001");
+			row1.Add(SummaryTableElement.Type.CsvMetadata, "platform", "windows");
+			row1.Add(SummaryTableElement.Type.CsvStatAverage, "frametime", 16.0);
+			table.AddRowData(row1, true, false, false, false);
+
+			var row2 = new SummaryTableRowData();
+			row2.Add(SummaryTableElement.Type.CsvMetadata, "csvid", "run_002");
+			row2.Add(SummaryTableElement.Type.CsvMetadata, "platform", "windows");
+			row2.Add(SummaryTableElement.Type.CsvStatAverage, "frametime", 20.0);
+			table.AddRowData(row2, true, false, false, false);
+
+			var sorted = table.SortRows(new List<string> { "platform" }, false);
+			var collated = sorted.CollateSortedTable(
+				new List<string> { "platform" }, true,
+				StringCollationVisibility.Auto, DateCollationVisibility.Newest);
+
+			string content = WriteAndRead(collated.ToFlatJsonDict(), "collated.json");
+			Assert.Contains("windows", content);
+			Assert.Contains("Avg frametime", content);
+			Assert.Contains("Min frametime", content);
+			Assert.Contains("Max frametime", content);
+		}
 	}
 
 	public class TableUtilTests
diff --git a/Engine/Source/Programs/CSVTools/PerfReportTool/PerfReportTool.cs b/Engine/Source/Programs/CSVTools/PerfReportTool/PerfReportTool.cs
--- a/Engine/Source/Programs/CSVTools/PerfReportTool/PerfReportTool.cs
+++ b/Engine/Source/Programs/CSVTools/PerfReportTool/PerfReportTool.cs
@@ -24,11 +24,19 @@ namespace PerfReportTool
 	class Version
 	{
 		// Format: Major.Minor.Bugfix
-		private static string VersionString = "4.272.1";
+		private static string VersionString = "4.273.0";
 
 		public static string Get() { return VersionString; }
 	};
 
+	[Flags]
+	enum SummaryTableOutputFormat
+	{
+		Html		= 1 << 0,
+		Csv			= 1 << 1,
+		Json		= 1 << 2
+	}
+
 	class HashHelper
 	{
 		public static string StringToHashStr(string strIn, int maxCharsOut = -1)
@@ -127,6 +135,7 @@ namespace PerfReportTool
 			"Optional bulk mode args: (use with -csvdir, -summaryTableCacheIn, -csvList, -prcList)\n" +
 			"  -recurse \n" +
 			"  -searchpattern <pattern>, e.g -searchpattern csvprofile*\n" +
+			"  -summaryTableOutputFormats <html|csv|json>\n : determines the output format(s) of the summary tables (default=html)" +
 			"  -customTable <comma separated fields> - sets or overrides the summary table metrics filter\n" +
 			"  -customTableSort <comma separated field row sort order> - overrides the summary table row sort\n" +
 			"  -noDetailedReports : skips individual report generation\n" +
@@ -134,7 +143,7 @@ namespace PerfReportTool
 			"  -collateTable : writes a collated table in addition to the main one, merging by row sort\n" +
 			"  -collateTableOnly : as -collateTable, but doesn't write the standard summary table.\n" +
 			"  -emailTable : writes a condensed email-friendly table (see the 'condensed' summary table)\n" +
-			"  -csvTable : writes the summary table in CSV format instead of html\n" +
+			"  -csvTable : (Deprecated, use -summaryTableOutputFormats csv) writes the summary table in CSV format instead of html\n" +
 			"  -summaryTableXML <XML filename>\n" +
 			"  -summaryTable <name> :\n" +
 			"     Selects a custom summary table type from the list in reportTypes.xml \n" +
@@ -210,13 +219,14 @@ namespace PerfReportTool
 			"  -regressionOutlierStdDevThreshold <n> (default = 4) : stddiv threshold for outliers (these are ignored)\n" +
 			"\n" +
 			"Json serialization:\n" +
-			"  -summaryTableToJson <path> : path (usually a json filename) to write summary table row data to\n" +
+			"  -summaryTableToJson <path> : path (usually a json filename) to write summary table row data to. This will write out the full csvs as json.\n" +
+			"								Use `-summaryTableOutputFormats json` to output a filtered summary table containing only the stats in the provided summary table.\n" +
 			"  -summaryTableToJsonSeparateFiles : writes separate files. -summaryTableToJson specifies the directory name\n" +
-			"  -summaryTableToJsonFastMode : exit after serializing json data (skips making summary tables)\n" +
+			"  -summaryTableToJsonFastMode : exit after serializing json data (skips making summary tables).\n" +
 			"  -summaryTableToJsonWriteAllElementData : write all element data, including tooltips, flags\n" +
 			"  -summaryTableToJsonMetadataOnly : only write CsvMetadata elements to json\n" +
-			"  -summaryTableToJsonFileStream : use a file stream to write Json. Experimental but can avoid OOMs\n" +
-			"  -summaryTableToJsonNoIndent : don't indent json output files\n" +
+			"  -summaryTableToJsonFileStream : use a file stream to write Json. Experimental but can avoid OOMs (also applies to `-summaryTableOutputFormats json`)\n" +
+			"  -summaryTableToJsonNoIndent : don't indent json output files (also applies to `-summaryTableOutputFormats json`)\n" +
 			"  -jsonToPrcs <json filename> : write PRCs. PRC files will be written to -summaryTableCache folder\n" +
 			"Telemetry Emission:\n" +
 			"  -summaryTableDataToTelemetry : publish the summary table as a telemetry event.\n" +
@@ -816,10 +826,16 @@ namespace PerfReportTool
 				{
 					summaryTableFilename = summaryTableFilename.Substring(0, summaryTableFilename.Length - 5);
 				}
-				bool bCsvTable = GetBoolArg("csvTable");
 				bool bCollateTable = GetBoolArg("collateTable");
 				bool bCollateTableOnly = GetBoolArg("collateTableOnly");
 
+				SummaryTableOutputFormat outputFormats = GetEnumArg<SummaryTableOutputFormat>("summaryTableOutputFormats", SummaryTableOutputFormat.Html);
+				if (GetBoolArg("csvTable"))
+				{
+					// Force exclusive for backwards compatiblity with previous behavior.
+					outputFormats = SummaryTableOutputFormat.Csv;
+				}
+
 				bCollateTable |= bCollateTableOnly;
 				string collatedTableFilename = summaryTableFilename + (bCollateTableOnly ? "" : "_Collated");
 
@@ -831,11 +847,11 @@ namespace PerfReportTool
 				SummaryTableInfo tableInfo = GetTableInfo(summaryTableName, customSummaryTableFilter, customSummaryTableRowSort);
 				if (!bCollateTableOnly)
 				{
-					WriteSummaryTableReport(outputDir, summaryTableFilename, summaryTable, tableInfo, false, bCsvTable, bSpreadsheetFriendlyStrings, null, telemetrySessionLabel, telemetryTestIdentity);
+					WriteSummaryTableReport(outputDir, summaryTableFilename, summaryTable, tableInfo, false, outputFormats, bSpreadsheetFriendlyStrings, null, telemetrySessionLabel, telemetryTestIdentity);
 				}
 				if (bCollateTable)
 				{
-					WriteSummaryTableReport(outputDir, collatedTableFilename, summaryTable, tableInfo, true, bCsvTable, bSpreadsheetFriendlyStrings, weightByColumnName, telemetrySessionLabel, telemetryTestIdentity);
+					WriteSummaryTableReport(outputDir, collatedTableFilename, summaryTable, tableInfo, true, outputFormats, bSpreadsheetFriendlyStrings, weightByColumnName, telemetrySessionLabel, telemetryTestIdentity);
 				}
 
 				// EmailTable is hardcoded to use the condensed type
@@ -847,7 +863,7 @@ namespace PerfReportTool
 						condensedSummaryTableName = "condensed";
 					}
 					tableInfo = GetTableInfo(condensedSummaryTableName, customSummaryTableFilter, customSummaryTableRowSort);
-					WriteSummaryTableReport(outputDir, summaryTableFilename + "_Email", summaryTable, tableInfo, true, false, bSpreadsheetFriendlyStrings, weightByColumnName);
+					WriteSummaryTableReport(outputDir, summaryTableFilename + "_Email", summaryTable, tableInfo, true, SummaryTableOutputFormat.Html, bSpreadsheetFriendlyStrings, weightByColumnName);
 				}
 				perfLog.LogTiming("WriteSummaryTable");
 			}
@@ -901,15 +917,15 @@ namespace PerfReportTool
 			return tableInfo;
 		}
 
-		void WriteSummaryTableReport(string outputDir, string filenameWithoutExtension, SummaryTable table, List<string> columnFilterList, List<string> rowSortList, bool bCollated, bool bToCSV, bool bSpreadsheetFriendlyStrings, List<SummarySectionBoundaryInfo> sectionBoundaries, string weightByColumnName)
+		void WriteSummaryTableReport(string outputDir, string filenameWithoutExtension, SummaryTable table, List<string> columnFilterList, List<string> rowSortList, bool bCollated, SummaryTableOutputFormat outputFormats, bool bSpreadsheetFriendlyStrings, List<SummarySectionBoundaryInfo> sectionBoundaries, string weightByColumnName)
 		{
 			SummaryTableInfo tableInfo = new SummaryTableInfo();
 			tableInfo.columnFilterList = columnFilterList;
 			tableInfo.rowSortList = rowSortList;
-			WriteSummaryTableReport(outputDir, filenameWithoutExtension, table, tableInfo, bCollated, bToCSV, bSpreadsheetFriendlyStrings, weightByColumnName);
+			WriteSummaryTableReport(outputDir, filenameWithoutExtension, table, tableInfo, bCollated, outputFormats, bSpreadsheetFriendlyStrings, weightByColumnName);
 		}
 
-		void WriteSummaryTableReport(string outputDir, string filenameWithoutExtension, SummaryTable table, SummaryTableInfo tableInfo, bool bCollated, bool bToCSV, bool bSpreadsheetFriendlyStrings, string weightByColumnNameOverride, string telemetrySessionLabel = null, string telemetryTestIdentity = null)
+		void WriteSummaryTableReport(string outputDir, string filenameWithoutExtension, SummaryTable table, SummaryTableInfo tableInfo, bool bCollated, SummaryTableOutputFormat outputFormats, bool bSpreadsheetFriendlyStrings, string weightByColumnNameOverride, string telemetrySessionLabel = null, string telemetryTestIdentity = null)
 		{
 			string weightByColumnName = weightByColumnNameOverride != null ? weightByColumnNameOverride : tableInfo.weightByColumn;
 			if (GetBoolArg("noWeightedAvg"))
@@ -969,11 +985,13 @@ namespace PerfReportTool
 				DateCollationVisibility dateCollationVisibility = GetEnumArg("collatedDateVisibility", tableInfo.dateCollationVisibility);
 				filteredTable = filteredTable.CollateSortedTable(tableInfo.rowSortList, addMinMaxColumns, stringCollationVisibility, dateCollationVisibility);
 			}
-			if (bToCSV)
+
+			if (outputFormats.HasFlag(SummaryTableOutputFormat.Csv))
 			{
 				filteredTable.WriteToCSV(filenameWithoutExtension + ".csv");
 			}
-			else
+
+			if (outputFormats.HasFlag(SummaryTableOutputFormat.Html))
 			{
 				filteredTable.ApplyDisplayNameMapping(statDisplayNameMapping);
 				string VersionString = GetBoolArg("noWatermarks") ? "" : Version.Get();
@@ -1004,6 +1022,12 @@ namespace PerfReportTool
 				);
 			}
 
+			if (outputFormats.HasFlag(SummaryTableOutputFormat.Json))
+			{
+				Dictionary<string, dynamic> jsonDict = filteredTable.ToFlatJsonDict();
+				SummaryTableDataJsonWriteHelper.WriteJsonFile(GetBoolArg("summaryTableToJsonFileStream"), filenameWithoutExtension + ".json", jsonDict, !GetBoolArg("summaryTableToJsonNoIndent"));
+			}
+
 			bool summaryTableDataToTelemetry = GetBoolArg("summaryTableDataToTelemetry") || GetBoolArg("summaryTableDataToTelemetryVerbose");
 
 			// Emit telemetry if enabled
diff --git a/Engine/Source/Programs/CSVTools/PerfReportTool/SummaryTable.cs b/Engine/Source/Programs/CSVTools/PerfReportTool/SummaryTable.cs
--- a/Engine/Source/Programs/CSVTools/PerfReportTool/SummaryTable.cs
+++ b/Engine/Source/Programs/CSVTools/PerfReportTool/SummaryTable.cs
@@ -1665,6 +1665,66 @@ namespace PerfSummaries
 			return newTable;
 		}
 
+		public Dictionary<string, dynamic> ToFlatJsonDict()
+		{
+			Dictionary<string, dynamic> result = new Dictionary<string, dynamic>();
+
+			SummaryTableColumn csvIdColumn = columnLookup.Get("csvid");
+			List<SummaryTableColumn> sortByColumns = columns.Where(c => c.isSortByColumn).ToList();
+
+			for (int rowIndex = 0; rowIndex < rowCount; rowIndex++)
+			{
+				string rowKey;
+				if (isCollated && sortByColumns.Count > 0)
+				{
+					rowKey = string.Join("_", sortByColumns.Select(c => c.GetStringValue(rowIndex)));
+				}
+				else if (csvIdColumn != null)
+				{
+					rowKey = csvIdColumn.GetStringValue(rowIndex);
+				}
+				else
+				{
+					rowKey = rowIndex.ToString();
+				}
+
+				Dictionary<string, dynamic> rowDict = new Dictionary<string, dynamic>();
+				foreach (SummaryTableColumn column in columns)
+				{
+					if (!column.isVisible)
+					{
+						continue;
+					}
+					if (column.isNumeric)
+					{
+						double val = column.GetValue(rowIndex);
+						if (val == double.MaxValue)
+						{
+							continue;
+						}
+						try
+						{
+							rowDict[column.name] = new decimal(val);
+						}
+						catch (OverflowException)
+						{
+							rowDict[column.name] = val;
+						}
+					}
+					else
+					{
+						string val = column.GetStringValue(rowIndex);
+						if (!string.IsNullOrEmpty(val))
+						{
+							rowDict[column.name] = val;
+						}
+					}
+				}
+				result[rowKey] = rowDict;
+			}
+			return result;
+		}
+
 		// Finds a particular aggregate column corresponding to a specified column
 		private SummaryTableColumn GetAggregateColumn(SummaryTableColumn inColumn, ColumnAggregateType aggregateType)
 		{
diff --git a/Engine/Source/Programs/CSVTools/PerfReportTool/SummaryTableData.cs b/Engine/Source/Programs/CSVTools/PerfReportTool/SummaryTableData.cs
--- a/Engine/Source/Programs/CSVTools/PerfReportTool/SummaryTableData.cs
+++ b/Engine/Source/Programs/CSVTools/PerfReportTool/SummaryTableData.cs
@@ -60,7 +60,7 @@ namespace PerfSummaries
 
 		}
 
-		private static void WriteJsonFile(bool bSerializeToStream, string filename, Dictionary<string, dynamic> dictToSerialize, bool bWriteIndented)
+		public static void WriteJsonFile(bool bSerializeToStream, string filename, Dictionary<string, dynamic> dictToSerialize, bool bWriteIndented)
 		{
 			JsonSerializerOptions options = new JsonSerializerOptions { WriteIndented = bWriteIndented, NumberHandling = System.Text.Json.Serialization.JsonNumberHandling.AllowNamedFloatingPointLiterals };
 
